Create the subobject Editor module that utilizes the new Subobject Data Interface instead of modifying subobjects directly in slate code.

#rnx

diff --git a/Engine/Source/Editor/Kismet/Public/SCSEditorExtensionContext.h b/Engine/Source/Editor/Kismet/Public/SCSEditorExtensionContext.h
--- a/Engine/Source/Editor/Kismet/Public/SCSEditorExtensionContext.h
+++ b/Engine/Source/Editor/Kismet/Public/SCSEditorExtensionContext.h
@@ -6,17 +6,29 @@
 #include "SCSEditorExtensionContext.generated.h"
 
 class SSCSEditor;
+class SSubobjectEditor;
+class SSubobjectInstanceEditor;
+class SSubobjectBlueprintEditor;
 
+// #TODO_BH Rename this file to Subobject Extension context
 UCLASS()
 class KISMET_API USCSEditorExtensionContext : public UObject
 {
 	GENERATED_BODY()
 
 public:
+
 	const TWeakPtr<SSCSEditor>& GetSCSEditor() const { return SCSEditor; }
+	
+	const TWeakPtr<SSubobjectEditor>& GetSubobjectEditor() const { return SubobjectEditor; }
 
 private:
 	friend SSCSEditor;
+	friend SSubobjectEditor;
+	friend SSubobjectInstanceEditor;
+	friend SSubobjectBlueprintEditor;
 
 	TWeakPtr<SSCSEditor> SCSEditor;
+	TWeakPtr<SSubobjectEditor> SubobjectEditor;
+
 };
diff --git a/Engine/Source/Editor/Kismet/Public/SSCSEditorMenuContext.h b/Engine/Source/Editor/Kismet/Public/SSCSEditorMenuContext.h
--- a/Engine/Source/Editor/Kismet/Public/SSCSEditorMenuContext.h
+++ b/Engine/Source/Editor/Kismet/Public/SSCSEditorMenuContext.h
@@ -7,14 +7,18 @@
 #include "SSCSEditorMenuContext.generated.h"
 
 class SSCSEditor;
+class SSubobjectEditor;
 
+// #TODO_BH Rename this for subobject editor menu context
 UCLASS()
 class KISMET_API USSCSEditorMenuContext : public UObject
 {
 	GENERATED_BODY()
 public:
-	
+
 	TWeakPtr<SSCSEditor> SCSEditor;
 
+	TWeakPtr<SSubobjectEditor> SubobjectEditor;
+
 	bool bOnlyShowPasteOption;
 };
diff --git a/Engine/Source/Editor/LevelEditor/LevelEditor.Build.cs b/Engine/Source/Editor/LevelEditor/LevelEditor.Build.cs
--- a/Engine/Source/Editor/LevelEditor/LevelEditor.Build.cs
+++ b/Engine/Source/Editor/LevelEditor/LevelEditor.Build.cs
@@ -82,6 +82,8 @@ public class LevelEditor : ModuleRules
 				"DesktopPlatform",
 				"DataLayerEditor",
 				"TranslationEditor",
+				"SubobjectEditor",
+				"SubobjectDataInterface",
 			}
 		);
 
diff --git a/Engine/Source/Editor/SubobjectEditor/Private/SSubobjectBlueprintEditor.cpp b/Engine/Source/Editor/SubobjectEditor/Private/SSubobjectBlueprintEditor.cpp
new file mode 100644
--- /dev/null
+++ b/Engine/Source/Editor/SubobjectEditor/Private/SSubobjectBlueprintEditor.cpp
@@ -0,0 +1,738 @@
+// Copyright Epic Games, Inc. All Rights Reserved.
+
+#include "SSubobjectBlueprintEditor.h"
+#include "ScopedTransaction.h"
+#include "IDocumentation.h"
+#include "GraphEditorActions.h"
+#include "Editor/EditorEngine.h"
+#include "ISCSEditorUICustomization.h"	// #TODO_BH Rename this to subobject
+#include "SCSEditorExtensionContext.h"	// #TODO_BH Rename this to subobject (this one will break VK, for sure)
+#include "ToolMenus.h"
+
+#include "Styling/SlateIconFinder.h"
+#include "SlateOptMacros.h"
+#include "Widgets/Input/SSearchBox.h"
+#include "Framework/MultiBox/MultiBoxBuilder.h"
+#include "SComponentClassCombo.h"
+#include "SEditorHeaderButton.h"
+#include "Editor/UnrealEdEngine.h"
+#include "Subsystems/PanelExtensionSubsystem.h"	// SExtensionPanel
+#include "ThumbnailRendering/ThumbnailManager.h"
+#include "Dialogs/Dialogs.h"					// FSuppressableWarningDialog
+#include "Kismet2/KismetEditorUtilities.h"
+#include "Kismet2/BlueprintEditorUtils.h"
+#include "Kismet2/ComponentEditorUtils.h"
+#include "Kismet2/ChildActorComponentEditorUtils.h"
+#include "ObjectTools.h"						// ThumbnailTools::CacheEmptyThumbnail
+#include "K2Node_ComponentBoundEvent.h"
+
+#define LOCTEXT_NAMESPACE "SSubobjectBlueprintEditor"
+
+extern UNREALED_API UUnrealEdEngine* GUnrealEd;
+
+////////////////////////////////////////////////
+// SSubobjectBlueprintEditor
+
+BEGIN_SLATE_FUNCTION_BUILD_OPTIMIZATION
+
+void SSubobjectBlueprintEditor::Construct(const FArguments& InArgs)
+{
+	ObjectContext = InArgs._ObjectContext;
+	PreviewActor = InArgs._PreviewActor;
+	OnSelectionUpdated = InArgs._OnSelectionUpdated;
+	OnItemDoubleClicked = InArgs._OnItemDoubleClicked;
+	AllowEditing = InArgs._AllowEditing;
+	bAllowTreeUpdates = true;
+	
+	CreateCommandList();
+
+	// Build the tree widget
+	FSlateBrush const* MobilityHeaderBrush = FEditorStyle::GetBrush(TEXT("ClassIcon.ComponentMobilityHeaderIcon"));
+	
+	ConstructTreeWidget();
+
+	// Should only be true when used in the blueprints details panel
+	const bool bInlineSearchBarWithButtons = ShowInlineSearchWithButtons();
+
+	TSharedPtr<SWidget> Contents;
+
+	TSharedPtr<SVerticalBox> HeaderBox;
+	TSharedPtr<SWidget> SearchBar = SAssignNew(FilterBox, SSearchBox)
+		.HintText(!bInlineSearchBarWithButtons ? LOCTEXT("SearchComponentsHint", "Search Components") : LOCTEXT("SearchHint", "Search"))
+		.OnTextChanged(this, &SSubobjectBlueprintEditor::OnFilterTextChanged)
+		.Visibility(this, &SSubobjectBlueprintEditor::GetComponentsFilterBoxVisibility);
+
+	FMenuBuilder EditBlueprintMenuBuilder = CreateMenuBuilder();
+
+	USCSEditorExtensionContext* ExtensionContext = NewObject<USCSEditorExtensionContext>();
+	ExtensionContext->SubobjectEditor = SharedThis(this);
+	ExtensionContext->AddToRoot();
+
+	ButtonBox = SNew(SHorizontalBox)
+	+ SHorizontalBox::Slot()
+	.VAlign(VAlign_Center)
+	.Padding(0.0f, 0.0f, 4.0f, 0.0f)
+	.AutoWidth()
+	[
+		SNew(SComponentClassCombo)
+		.AddMetaData<FTagMetaData>(FTagMetaData(TEXT("Actor.AddComponent")))
+		.Visibility(HideComponentClassCombo.Get() ? EVisibility::Hidden : EVisibility::Visible)
+		.OnSubobjectClassSelected(this, &SSubobjectBlueprintEditor::PerformComboAddClass)
+		.ToolTipText(LOCTEXT("AddComponent_Tooltip", "Adds a new component to this actor"))
+		.IsEnabled(true)
+	]
+	+ SHorizontalBox::Slot()
+	.VAlign(VAlign_Center)
+	.Padding(0.0f, 0.0f, 4.0f, 0.0f)
+	.AutoWidth()
+	[
+		SAssignNew(ExtensionPanel, SExtensionPanel)
+		.ExtensionPanelID("SCSEditor.NextToAddComponentButton")
+		.ExtensionContext(ExtensionContext)
+	]
+	
+	// horizontal slot index #2 => reserved for BP-editor search bar (see 'ButtonBox' and 'SearchBarHorizontalSlotIndex' usage below)
+
+	+ SHorizontalBox::Slot()
+	.VAlign(VAlign_Center)
+	.Padding(0.0f, 0.0f, 4.0f, 0.0f)
+	.AutoWidth()
+	[
+		SNew(SEditorHeaderButton)
+		.AddMetaData<FTagMetaData>(FTagMetaData(TEXT("Actor.ConvertToBlueprint")))
+		.Visibility(this, &SSubobjectBlueprintEditor::GetPromoteToBlueprintButtonVisibility)
+		.OnClicked(this, &SSubobjectBlueprintEditor::OnPromoteToBlueprintClicked)
+		.Icon(FAppStyle::Get().GetBrush("Icons.Blueprints"))
+		.ToolTip(IDocumentation::Get()->CreateToolTip(
+			TAttribute<FText>(LOCTEXT("PromoteToBluerprintTooltip", "Converts this actor into a reusable Blueprint Class that can have script behavior")),
+			nullptr,
+			TEXT("Shared/LevelEditor"),
+			TEXT("ConvertToBlueprint")))
+	]
+	+ SHorizontalBox::Slot()
+	.VAlign(VAlign_Center)
+	.AutoWidth()
+	[
+		SNew(SEditorHeaderButton)
+		.AddMetaData<FTagMetaData>(FTagMetaData(TEXT("Actor.EditBlueprint")))
+		.Visibility(this, &SSubobjectBlueprintEditor::GetEditBlueprintButtonVisibility)
+		.ToolTipText(LOCTEXT("EditActorBlueprint_Tooltip", "Edit the Blueprint for this Actor"))
+		.Icon(FAppStyle::Get().GetBrush("Icons.Blueprints"))
+		.MenuContent()
+		[
+			EditBlueprintMenuBuilder.MakeWidget()
+		]
+	];
+
+	Contents = SNew(SVerticalBox)
+	+ SVerticalBox::Slot()
+	.AutoHeight()
+	.VAlign(VAlign_Top)
+	.Padding(4.f, 0, 4.f, 4.f)
+	[
+		SAssignNew(HeaderBox, SVerticalBox)
+	]
+
+	+ SVerticalBox::Slot()
+	[
+		SNew(SBorder)
+		.BorderImage(FAppStyle::Get().GetBrush("SCSEditor.Background"))
+		.Padding(4.f)
+		.AddMetaData<FTagMetaData>(FTagMetaData(TEXT("ComponentsPanel")))
+		.Visibility(this, &SSubobjectBlueprintEditor::GetComponentsTreeVisibility)
+		[
+			TreeWidget.ToSharedRef()
+		]
+	];
+
+	// Only insert the buttons and search bar in the Blueprints version
+	if (bInlineSearchBarWithButtons) // Blueprints
+	{
+		ButtonBox->AddSlot()
+			.FillWidth(1.0f)
+			.VAlign(VAlign_Center)
+			.Padding(3.0f, 3.0f)
+		[
+			SearchBar.ToSharedRef()
+		];
+
+		HeaderBox->AddSlot()
+			.VAlign(VAlign_Center)
+			.AutoHeight()
+		[
+			ButtonBox.ToSharedRef()
+		];
+	}
+
+	this->ChildSlot
+	[
+		Contents.ToSharedRef()
+	];
+
+	// Update tree method? 
+	UpdateTree();
+}
+
+END_SLATE_FUNCTION_BUILD_OPTIMIZATION
+
+void SSubobjectBlueprintEditor::OnDeleteNodes()
+{
+	const FScopedTransaction Transaction( LOCTEXT("RemoveComponents", "Remove Components") );
+
+	// Invalidate any active component in the visualizer
+	GUnrealEd->ComponentVisManager.ClearActiveComponentVis();
+
+	UBlueprint* Blueprint = GetBlueprint();
+	check(Blueprint != nullptr);
+	
+	// Get the current render info for the blueprint. If this is NULL then the blueprint is not currently visualizable (no visible primitive components)
+	FThumbnailRenderingInfo* RenderInfo = GUnrealEd->GetThumbnailManager()->GetRenderingInfo( Blueprint );
+
+	// A lamda for displaying a confirm message to the user if there is a dynamic delegate bound to the 
+	// component they are trying to delete
+	auto ConfirmDeleteLambda = [](const FSubobjectData* Data) -> FSuppressableWarningDialog::EResult
+	{
+		if (ensure(Data))
+		{
+			FText VarNam = FText::FromName(Data->GetVariableName());
+			FText ConfirmDelete = FText::Format(LOCTEXT("ConfirmDeleteDynamicDelegate", "Component \"{0}\" has bound events in use! If you delete it then those nodes will become invalid. Are you sure you want to delete it?"), VarNam);
+
+			// Warn the user that this may result in data loss
+			FSuppressableWarningDialog::FSetupInfo Info(ConfirmDelete, LOCTEXT("DeleteComponent", "Delete Component"), "DeleteComponentInUse_Warning");
+			Info.ConfirmText = LOCTEXT("ConfirmDeleteDynamicDelegate_Yes", "Yes");
+			Info.CancelText = LOCTEXT("ConfirmDeleteDynamicDelegate_No", "No");
+
+			FSuppressableWarningDialog DeleteVariableInUse(Info);
+
+			// If the user selects cancel then return false
+			return DeleteVariableInUse.ShowModal();
+		}
+
+		return FSuppressableWarningDialog::Cancel;
+	};
+	
+	// Gather the handles of the components that we want to delete
+	TArray<FSubobjectDataHandle> HandlesToDelete;
+	TArray<FSubobjectEditorTreeNodePtrType> SelectedNodes = GetSelectedNodes();
+	for(const FSubobjectEditorTreeNodePtrType& Node : SelectedNodes)
+	{
+		check(Node->IsValid());
+		
+		if(const FSubobjectData* Data = Node->GetDataSource())
+		{
+			// If this node is in use by Dynamic delegates, then confirm before continuing
+			if (FKismetEditorUtilities::PropertyHasBoundEvents(Blueprint, Data->GetVariableName()))
+			{
+				// The user has decided not to delete the component, stop trying to delete this component
+				if (ConfirmDeleteLambda(Data) == FSuppressableWarningDialog::Cancel)
+				{
+					return;
+				}
+			}
+			HandlesToDelete.Add(Data->GetHandle());
+		}
+	}
+
+	if(USubobjectDataSubsystem* System = USubobjectDataSubsystem::Get())
+	{
+		FSubobjectDataHandle HandleToSelect;
+		
+		int32 NumDeleted = System->DeleteSubobjects(RootNodes[0]->GetDataHandle(), HandlesToDelete, HandleToSelect, GetBlueprint());
+
+		if(NumDeleted > 0)
+		{
+			if(HandleToSelect.IsValid())
+			{
+				FSubobjectEditorTreeNodePtrType NodeToSelect = FindSlateNodeForHandle(HandleToSelect);
+				if(NodeToSelect.IsValid())
+				{
+					TreeWidget->SetSelection(NodeToSelect);
+				}
+			}
+			
+			UpdateTree();
+
+			// If we had a thumbnail before we deleted any components, check to see if we should clear it
+			// If we deleted the final visualizable primitive from the blueprint, GetRenderingInfo should return NULL
+			FThumbnailRenderingInfo* NewRenderInfo = GUnrealEd->GetThumbnailManager()->GetRenderingInfo(Blueprint);
+			if (RenderInfo && !NewRenderInfo)
+			{
+				// We removed the last visible primitive component, clear the thumbnail
+				const FString BPFullName = FString::Printf(TEXT("%s %s"), *Blueprint->GetClass()->GetName(), *Blueprint->GetPathName());
+				UPackage* BPPackage = Blueprint->GetOutermost();
+				ThumbnailTools::CacheEmptyThumbnail( BPFullName, BPPackage );
+			}
+			// Do this AFTER marking the Blueprint as modified
+			UpdateSelectionFromNodes(TreeWidget->GetSelectedItems());
+		}
+	}
+}
+
+void SSubobjectBlueprintEditor::CopySelectedNodes()
+{
+	TArray<FSubobjectDataHandle> SelectedHandles = GetSelectedHandles();
+	if(USubobjectDataSubsystem* System = USubobjectDataSubsystem::Get())
+	{
+		return System->CopySubobjects(SelectedHandles, GetBlueprint());
+	}
+}
+
+void SSubobjectBlueprintEditor::OnDuplicateComponent()
+{
+	TArray<FSubobjectDataHandle> SelectedNodes = GetSelectedHandles();
+	if(SelectedNodes.Num() > 0)
+	{
+		// Force the text box being edited (if any) to commit its text. The duplicate operation may trigger a regeneration of the tree view,
+		// releasing all row widgets. If one row was in edit mode (rename/rename on create), it was released before losing the focus and
+		// this would prevent the completion of the 'rename' or 'create + give initial name' transaction (occurring on focus lost).
+		FSlateApplication::Get().ClearKeyboardFocus();
+		
+		const FScopedTransaction Transaction(SelectedNodes.Num() > 1 ? LOCTEXT("DuplicateComponents", "Duplicate Components") : LOCTEXT("DuplicateComponent", "Duplicate Component"));
+		
+		if(USubobjectDataSubsystem* System = USubobjectDataSubsystem::Get())
+		{
+			System->DuplicateSubobjects(GetObjectContextHandle(), SelectedNodes, GetBlueprint());
+			UpdateTree();
+		}
+	}
+}
+
+bool SSubobjectBlueprintEditor::CanPasteNodes() const
+{
+	if(!IsEditingAllowed())
+	{
+		return false;
+	}
+	
+	if(FSubobjectEditorTreeNodePtrType SceneRoot = GetSceneRootNode())
+	{
+		if(USubobjectDataSubsystem* System = USubobjectDataSubsystem::Get())
+		{
+			return System->CanPasteSubobjects(SceneRoot->GetDataHandle(), GetBlueprint());
+		}
+	}
+
+	return false;
+}
+
+void SSubobjectBlueprintEditor::PasteNodes()
+{
+	if(USubobjectDataSubsystem* System = USubobjectDataSubsystem::Get())
+	{
+		TArray<FSubobjectDataHandle> OutHandles;
+		
+		// stop allowing tree updates during paste, otherwise MarkBlueprintAsModified will trigger in the middle of it
+		bool bRestoreAllowTreeUpdates = bAllowTreeUpdates;
+		bAllowTreeUpdates = false;
+		TArray<FSubobjectDataHandle> HandlesToPasteOnto = GetSelectedHandles();
+		if(HandlesToPasteOnto.IsEmpty())
+		{
+			if(FSubobjectEditorTreeNodePtrType RootPtr = GetSceneRootNode())
+			{
+				if(RootPtr.IsValid())
+				{
+					HandlesToPasteOnto.Emplace(RootPtr->GetDataHandle());
+				}
+			}			
+		}
+		
+		System->PasteSubobjects(GetObjectContextHandle(), HandlesToPasteOnto, GetBlueprint(), OutHandles);
+
+		// allow tree updates again
+		bAllowTreeUpdates = bRestoreAllowTreeUpdates;
+		
+		if(OutHandles.Num() > 0)
+		{
+			// We only want the pasted node(s) to be selected
+			TreeWidget->ClearSelection();
+			UpdateTree();
+
+			for(const FSubobjectDataHandle& Handle : OutHandles)
+			{
+				if(FSubobjectEditorTreeNodePtrType SlateNode = FindSlateNodeForHandle(Handle))
+				{
+					TreeWidget->SetItemSelection(SlateNode, true);	
+				}
+			}
+		}
+	}
+}
+
+void SSubobjectBlueprintEditor::OnAttachToDropAction(FSubobjectEditorTreeNodePtrType DroppedOn, const TArray<FSubobjectEditorTreeNodePtrType>& DroppedNodePtrs)
+{
+	// Ask the subsystem to attach the dropped nodes onto the dropped on node
+	check(DroppedOn.IsValid());
+	check(DroppedNodePtrs.Num() > 0);
+	
+	USubobjectDataSubsystem* System = USubobjectDataSubsystem::Get();
+	check(System);
+	
+	const FScopedTransaction TransactionContext(DroppedNodePtrs.Num() > 1 ? LOCTEXT("AttachComponents", "Attach Components") : LOCTEXT("AttachComponent", "Attach Component"));
+	TArray<FSubobjectDataHandle> HandlesToMove;
+	for(const FSubobjectEditorTreeNodePtrType& DroppedNodePtr : DroppedNodePtrs)
+	{
+		HandlesToMove.Add(DroppedNodePtr->GetDataHandle());
+	}
+
+	FReparentSubobjectParams Params;
+	Params.NewParentHandle = DroppedOn->GetDataHandle();
+	Params.BlueprintContext =  GetBlueprint();
+	Params.ActorPreviewContext = GetActorPreview();
+
+	System->ReparentSubobjects(Params, HandlesToMove);
+
+	check(TreeWidget.IsValid());
+	TreeWidget->SetItemExpansion(DroppedOn, true);
+
+	PostDragDropAction(true);
+}
+
+void SSubobjectBlueprintEditor::OnDetachFromDropAction(const TArray<FSubobjectEditorTreeNodePtrType>& DroppedNodePtrs)
+{
+	check(DroppedNodePtrs.Num() > 0);
+
+	const FScopedTransaction TransactionContext(DroppedNodePtrs.Num() > 1 ? LOCTEXT("DetachComponents", "Detach Components") : LOCTEXT("DetachComponent", "Detach Component"));
+
+	TArray<FSubobjectDataHandle> HandlesToMove;
+	Utils::PopulateHandlesArray(DroppedNodePtrs, HandlesToMove);
+
+	// Attach the dropped node to the current scene root node
+	FSubobjectEditorTreeNodePtrType SceneRootNodePtr = GetSceneRootNode();
+	check(SceneRootNodePtr.IsValid());
+
+	// Ask the subsystem to reparent this object to the scene root
+	if (USubobjectDataSubsystem* System = USubobjectDataSubsystem::Get())
+	{
+		FReparentSubobjectParams Params;
+		Params.NewParentHandle = SceneRootNodePtr->GetDataHandle();
+		Params.BlueprintContext =  GetBlueprint();
+		Params.ActorPreviewContext = GetActorPreview();
+		System->ReparentSubobjects(Params, HandlesToMove);
+	}
+
+	PostDragDropAction(false);
+}
+
+void SSubobjectBlueprintEditor::OnMakeNewRootDropAction(FSubobjectEditorTreeNodePtrType DroppedNodePtr)
+{
+	// Get the current scene root node
+	FSubobjectEditorTreeNodePtrType SceneRootNodePtr = GetSceneRootNode();
+
+	// We cannot handle the drop action if any of these conditions fail on entry.
+	if (!ensure(SceneRootNodePtr.IsValid()))
+	{
+		return;
+	}
+	if(!ensure(DroppedNodePtr.IsValid()))
+	{
+		return;
+	}
+	
+	const FScopedTransaction TransactionContext(LOCTEXT("MakeNewSceneRoot", "Make New Scene Root"));
+	
+	USubobjectDataSubsystem* Subsystem = USubobjectDataSubsystem::Get();
+	const bool bSuccess = Subsystem->MakeNewSceneRoot(
+		GetObjectContextHandle(),
+		 DroppedNodePtr->GetDataHandle(),
+		 GetBlueprint());
+	
+	PostDragDropAction(true);
+}
+
+void SSubobjectBlueprintEditor::PostDragDropAction(bool bRegenerateTreeNodes)
+{
+	GUnrealEd->ComponentVisManager.ClearActiveComponentVis();
+
+	UpdateTree(bRegenerateTreeNodes);
+
+	FBlueprintEditorUtils::PostEditChangeBlueprintActors(GetBlueprint(), true);
+}
+
+TSharedPtr<SWidget> SSubobjectBlueprintEditor::BuildSceneRootDropActionMenu(FSubobjectEditorTreeNodePtrType DroppedOntoNodePtr, FSubobjectEditorTreeNodePtrType DroppedNodePtr)
+{
+	FMenuBuilder MenuBuilder(true, CommandList);
+
+	const FSubobjectData* DroppedNodeData = DroppedNodePtr->GetDataSource();
+	const FSubobjectData* DroppedOntoNodeData = DroppedOntoNodePtr->GetDataSource();
+
+	check(DroppedNodeData);
+	
+	MenuBuilder.BeginSection("SceneRootNodeDropActions", LOCTEXT("SceneRootNodeDropActionContextMenu", "Drop Actions"));
+	{
+		const FText DroppedVariableNameText = FText::FromName( DroppedNodeData->GetVariableName() );
+		
+		const FText NodeVariableNameText = FText::FromName( DroppedOntoNodeData->GetVariableName() );
+
+		bool bDroppedInSameBlueprint = true;
+
+		MenuBuilder.AddMenuEntry(
+			LOCTEXT("DropActionLabel_AttachToRootNode", "Attach"),
+			bDroppedInSameBlueprint 
+			? FText::Format( LOCTEXT("DropActionToolTip_AttachToRootNode", "Attach {0} to {1}."), DroppedVariableNameText, NodeVariableNameText )
+			: FText::Format( LOCTEXT("DropActionToolTip_AttachToRootNodeFromCopy", "Copy {0} to a new variable and attach it to {1}."), DroppedVariableNameText, NodeVariableNameText ),
+			FSlateIcon(),
+			FUIAction(
+				FExecuteAction::CreateSP(this, &SSubobjectEditor::OnAttachToDropAction, DroppedOntoNodePtr, DroppedNodePtr),
+				FCanExecuteAction()));
+
+		const bool bIsDefaultSceneRoot = DroppedNodeData->IsDefaultSceneRoot();
+
+		FText NewRootNodeText = bIsDefaultSceneRoot
+			? FText::Format(LOCTEXT("DropActionToolTip_MakeNewRootNodeAndDelete", "Make {0} the new root. The default root will be deleted."), DroppedVariableNameText)
+			: FText::Format(LOCTEXT("DropActionToolTip_MakeNewRootNode", "Make {0} the new root."), DroppedVariableNameText);
+
+		FText NewRootNodeFromCopyText = bIsDefaultSceneRoot
+			? FText::Format(LOCTEXT("DropActionToolTip_MakeNewRootNodeFromCopyAndDelete", "Copy {0} to a new variable and make it the new root. The default root will be deleted."), DroppedVariableNameText)
+			: FText::Format(LOCTEXT("DropActionToolTip_MakeNewRootNodeFromCopy", "Copy {0} to a new variable and make it the new root."), DroppedVariableNameText);
+
+		MenuBuilder.AddMenuEntry(
+			LOCTEXT("DropActionLabel_MakeNewRootNode", "Make New Root"),
+			bDroppedInSameBlueprint ? NewRootNodeText : NewRootNodeFromCopyText,
+			FSlateIcon(),
+			FUIAction(
+				FExecuteAction::CreateSP(this, &SSubobjectEditor::OnMakeNewRootDropAction, DroppedNodePtr),
+				FCanExecuteAction()));
+	}
+	MenuBuilder.EndSection();
+
+	return MenuBuilder.MakeWidget();
+}
+
+FSubobjectDataHandle SSubobjectBlueprintEditor::AddNewSubobject(const FSubobjectDataHandle& ParentHandle, UClass* NewClass, UObject* AssetOverride, FText& OutFailReason, TUniquePtr<FScopedTransaction> InOngoingTransaction)
+{
+	FAddNewSubobjectParams Params;
+    Params.ParentHandle = ParentHandle;
+    Params.NewClass = NewClass;
+    Params.AssetOverride = AssetOverride;
+
+    // Make sure to populate the blueprint context of what we are currently editing!
+    Params.BlueprintContext = GetBlueprint();
+    USubobjectDataSubsystem* System = USubobjectDataSubsystem::Get();
+    check(System);
+	
+	FSubobjectDataHandle NewHandle = System->AddNewSubobject(Params,OutFailReason);
+	
+    return NewHandle;	
+}
+
+void SSubobjectBlueprintEditor::PopulateContextMenuImpl(UToolMenu* InMenu, TArray<FSubobjectEditorTreeNodePtrType>& InSelectedItems, bool bIsChildActorSubtreeNodeSelected)
+{
+	FToolMenuSection& BlueprintSCSSection = InMenu->AddSection("BlueprintSCS");
+	if (InSelectedItems.Num() == 1)
+	{
+		BlueprintSCSSection.AddMenuEntry(FGraphEditorCommands::Get().FindReferences);
+	}
+
+	// Create an "Add Event" option in the context menu only if we can edit
+	// the currently selected objects
+	if (IsEditingAllowed())
+	{
+		// Collect the classes of all selected objects
+		TArray<UClass*> SelectionClasses;
+		for (auto NodeIter = InSelectedItems.CreateConstIterator(); NodeIter; ++NodeIter)
+		{
+			FSubobjectEditorTreeNodePtrType TreeNode = *NodeIter;
+			const FSubobjectData* Data = TreeNode->GetDataSource();
+			check(Data);
+			
+			if (const UActorComponent* ComponentTemplate = TreeNode->GetComponentTemplate())
+			{
+				// If the component is native then we need to ensure it can actually be edited before we display it
+				if (!Data->IsNativeComponent() || FComponentEditorUtils::GetPropertyForEditableNativeComponent(ComponentTemplate))
+				{
+					SelectionClasses.Add(ComponentTemplate->GetClass());
+				}
+			}
+		}
+
+		if (SelectionClasses.Num())
+		{
+			// Find the common base class of all selected classes
+			UClass* SelectedClass = UClass::FindCommonBase(SelectionClasses);
+			// Build an event submenu if we can generate events
+			if (FBlueprintEditorUtils::CanClassGenerateEvents(SelectedClass))
+			{
+				BlueprintSCSSection.AddSubMenu(
+                    "AddEventSubMenu",
+                    LOCTEXT("AddEventSubMenu", "Add Event"),
+                    LOCTEXT("ActtionsSubMenu_ToolTip", "Add Event"),
+                    FNewMenuDelegate::CreateStatic(&SSubobjectBlueprintEditor::BuildMenuEventsSection,
+                        GetBlueprint(), SelectedClass, FCanExecuteAction::CreateSP(this, &SSubobjectEditor::IsEditingAllowed),
+                        FGetSelectedObjectsDelegate::CreateSP(this, &SSubobjectEditor::GetSelectedItemsForContextMenu)));
+			}
+		}
+	}
+
+	TArray<UActorComponent*> SelectedComponents;
+	for(const FSubobjectEditorTreeNodePtrType& SelectedNodePtr : InSelectedItems)
+	{
+		check(SelectedNodePtr->IsValid());
+		// Get the component template associated with the selected node
+		const UActorComponent* ComponentTemplate = SelectedNodePtr->GetComponentTemplate();
+		if (ComponentTemplate)
+		{
+			// #TODO_BH Remove this const cast
+			SelectedComponents.Add(const_cast<UActorComponent*>(ComponentTemplate));
+		}
+	}
+	
+	// Common menu options added for all component types
+    FComponentEditorUtils::FillComponentContextMenuOptions(InMenu, SelectedComponents);
+
+	// For a selection outside of a child actor subtree, we may choose to include additional options
+	if (SelectedComponents.Num() == 1 && !bIsChildActorSubtreeNodeSelected)
+	{
+		// Extra options for a child actor component
+		if (UChildActorComponent* SelectedChildActorComponent = Cast<UChildActorComponent>(SelectedComponents[0]))
+		{
+			// These options will get added only in SCS mode
+			FChildActorComponentEditorUtils::FillComponentContextMenuOptions(InMenu, SelectedChildActorComponent);
+		}
+	}
+}
+
+FSubobjectEditorTreeNodePtrType SSubobjectBlueprintEditor::GetSceneRootNode() const
+{
+	if(RootNodes.Num() == 0)
+	{
+		return nullptr;
+	}
+	
+	// Loop through all our slate nodes and see if any are marked as scene root
+	if(USubobjectDataSubsystem* System = USubobjectDataSubsystem::Get())
+	{
+		FSubobjectDataHandle RootHandle = System->FindSceneRootForSubobject(RootNodes[0]->GetDataHandle());
+		if(RootHandle.IsValid())
+		{
+			return FindSlateNodeForObject(RootHandle.GetSharedDataPtr()->GetObject(), true);			
+		}
+	}
+	return nullptr;
+}
+
+void SSubobjectBlueprintEditor::BuildMenuEventsSection(FMenuBuilder& Menu, UBlueprint* Blueprint, UClass* SelectedClass, FCanExecuteAction CanExecuteActionDelegate, FGetSelectedObjectsDelegate GetSelectedObjectsDelegate)
+{
+	// Get Selected Nodes
+	TArray<FComponentEventConstructionData> SelectedNodes;
+	GetSelectedObjectsDelegate.ExecuteIfBound(SelectedNodes);
+
+	struct FMenuEntry
+	{
+		FText Label;
+		FText ToolTip;
+		FUIAction UIAction;
+	};
+
+	TArray<FMenuEntry> Actions;
+	TArray<FMenuEntry> NodeActions;
+	// Build Events entries
+	for (TFieldIterator<FMulticastDelegateProperty> PropertyIt(SelectedClass, EFieldIteratorFlags::IncludeSuper); PropertyIt; ++PropertyIt)
+	{
+		FMulticastDelegateProperty* Property = *PropertyIt;
+
+		// Check for multicast delegates that we can safely assign
+		if (!Property->HasAnyPropertyFlags(CPF_Parm) && Property->HasAllPropertyFlags(CPF_BlueprintAssignable))
+		{
+			FName EventName = Property->GetFName();
+			int32 ComponentEventViewEntries = 0;
+			// Add View Event Per Component
+			for (auto NodeIter = SelectedNodes.CreateConstIterator(); NodeIter; ++NodeIter)
+			{
+				if(NodeIter->Component.IsValid())
+				{
+					FName VariableName = NodeIter->VariableName;
+					FObjectProperty* VariableProperty = FindFProperty<FObjectProperty>( Blueprint->SkeletonGeneratedClass, VariableName );
+
+					if(VariableProperty && FKismetEditorUtilities::FindBoundEventForComponent(Blueprint, EventName, VariableProperty->GetFName()))
+					{
+						FMenuEntry NewEntry;
+						NewEntry.Label = (SelectedNodes.Num() > 1) ?	FText::Format( LOCTEXT("ViewEvent_ToolTipFor", "{0} for {1}"), FText::FromName(EventName), FText::FromName(VariableName)) : 
+																		FText::Format( LOCTEXT("ViewEvent_ToolTip", "{0}"), FText::FromName(EventName));
+						NewEntry.UIAction =	FUIAction(FExecuteAction::CreateStatic(&SSubobjectBlueprintEditor::ViewEvent, Blueprint, EventName, *NodeIter), CanExecuteActionDelegate);
+						NodeActions.Add(NewEntry);
+						ComponentEventViewEntries++;
+					}
+				}
+			}
+			if(ComponentEventViewEntries < SelectedNodes.Num())
+			{
+				// Create menu Add entry
+				FMenuEntry NewEntry;
+				NewEntry.Label = FText::Format(LOCTEXT("AddEvent_ToolTip", "Add {0}" ), FText::FromName(EventName));
+				NewEntry.UIAction =	FUIAction(FExecuteAction::CreateStatic(&SSubobjectBlueprintEditor::CreateEventsForSelection, Blueprint, EventName, GetSelectedObjectsDelegate), CanExecuteActionDelegate);
+				Actions.Add(NewEntry);
+			}
+		}
+	}
+	// Build Menu Sections
+	Menu.BeginSection("AddComponentActions", LOCTEXT("AddEventHeader", "Add Event"));
+	for (auto ItemIter = Actions.CreateConstIterator(); ItemIter; ++ItemIter )
+	{
+		Menu.AddMenuEntry( ItemIter->Label, ItemIter->ToolTip, FSlateIcon(), ItemIter->UIAction );
+	}
+	Menu.EndSection();
+	Menu.BeginSection("ViewComponentActions", LOCTEXT("ViewEventHeader", "View Existing Events"));
+	for (auto ItemIter = NodeActions.CreateConstIterator(); ItemIter; ++ItemIter )
+	{
+		Menu.AddMenuEntry( ItemIter->Label, ItemIter->ToolTip, FSlateIcon(), ItemIter->UIAction );
+	}
+	Menu.EndSection();
+}
+
+void SSubobjectBlueprintEditor::CreateEventsForSelection(UBlueprint* Blueprint, FName EventName, FGetSelectedObjectsDelegate GetSelectedObjectsDelegate)
+{
+	if (EventName != NAME_None)
+	{
+		TArray<FComponentEventConstructionData> SelectedNodes;
+		GetSelectedObjectsDelegate.ExecuteIfBound(SelectedNodes);
+
+		for (auto SelectionIter = SelectedNodes.CreateConstIterator(); SelectionIter; ++SelectionIter)
+		{
+			ConstructEvent(Blueprint, EventName, *SelectionIter);
+		}
+	}
+}
+
+void SSubobjectBlueprintEditor::ConstructEvent(UBlueprint* Blueprint, const FName EventName, const FComponentEventConstructionData EventData)
+{
+	// Find the corresponding variable property in the Blueprint
+	FObjectProperty* VariableProperty = FindFProperty<FObjectProperty>(Blueprint->SkeletonGeneratedClass, EventData.VariableName );
+
+	if (VariableProperty)
+	{
+		if (!FKismetEditorUtilities::FindBoundEventForComponent(Blueprint, EventName, VariableProperty->GetFName()))
+		{
+			FKismetEditorUtilities::CreateNewBoundEventForComponent(EventData.Component.Get(), EventName, Blueprint, VariableProperty);
+		}
+	}
+}
+
+void SSubobjectBlueprintEditor::ViewEvent(UBlueprint* Blueprint, const FName EventName, const FComponentEventConstructionData EventData)
+{
+	// Find the corresponding variable property in the Blueprint
+	FObjectProperty* VariableProperty = FindFProperty<FObjectProperty>(Blueprint->SkeletonGeneratedClass, EventData.VariableName);
+
+	if (VariableProperty)
+	{
+		const UK2Node_ComponentBoundEvent* ExistingNode = FKismetEditorUtilities::FindBoundEventForComponent(Blueprint, EventName, VariableProperty->GetFName());
+		if (ExistingNode)
+		{
+			FKismetEditorUtilities::BringKismetToFocusAttentionOnObject(ExistingNode);
+		}
+	}
+}
+
+FSlateColor SSubobjectBlueprintEditor::GetColorTintForIcon(FSubobjectEditorTreeNodePtrType Node) const
+{
+	static const FLinearColor IntroducedHereColor(FLinearColor::White);
+	
+	const FSubobjectData* Data = Node ? Node->GetDataSource() : nullptr;
+	if(!Data || Data->IsActor())
+	{
+		return IntroducedHereColor;
+	}
+
+	return Data->GetColorTintForIcon();
+}
+
+#undef LOCTEXT_NAMESPACE
diff --git a/Engine/Source/Editor/SubobjectEditor/Private/SSubobjectEditor.cpp b/Engine/Source/Editor/SubobjectEditor/Private/SSubobjectEditor.cpp
new file mode 100644
--- /dev/null
+++ b/Engine/Source/Editor/SubobjectEditor/Private/SSubobjectEditor.cpp
@@ -0,0 +1,3268 @@
+// Copyright Epic Games, Inc. All Rights Reserved.
+
+#include "SSubobjectEditor.h"
+#include "Editor.h"
+#include "ClassIconFinder.h"
+#include "ScopedTransaction.h"
+#include "ComponentAssetBroker.h"		// FComponentAssetBrokerage
+#include "TutorialMetaData.h"
+#include "IDocumentation.h"
+#include "Framework/Commands/GenericCommands.h"
+#include "Kismet2/ComponentEditorUtils.h"
+#include "Kismet2/BlueprintEditorUtils.h"
+#include "Kismet2/KismetEditorUtilities.h"
+#include "Kismet2/ChildActorComponentEditorUtils.h"
+#include "GraphEditorActions.h"
+#include "Editor/EditorEngine.h"
+#include "ISCSEditorUICustomization.h"	// #TODO_BH Rename this to subobject
+#include "SCSEditorExtensionContext.h"	// #TODO_BH Rename this to subobject (this one will break VK, for sure)
+#include "SSCSEditorMenuContext.h"		// #TODO_BH Rename this to subobject
+#include "Toolkits/ToolkitManager.h"
+#include "BlueprintEditor.h"
+#include "ToolMenus.h"
+#include "Components/ChildActorComponent.h"
+#include "Components/PrimitiveComponent.h"
+
+#include "Styling/SlateIconFinder.h"
+#include "SlateOptMacros.h"
+#include "Widgets/Text/SInlineEditableTextBlock.h"
+#include "Widgets/Input/SSearchBox.h"
+#include "Widgets/Images/SImage.h"
+#include "Widgets/SToolTip.h"
+#include "Widgets/Layout/SSpacer.h"
+#include "Selection.h"
+#include "Framework/MultiBox/MultiBoxBuilder.h"
+#include "SComponentClassCombo.h"
+#include "SEditorHeaderButton.h"
+#include "CreateBlueprintFromActorDialog.h"
+#include "ClassViewerFilter.h"
+#include "GameProjectGenerationModule.h"	// Adding new component classes
+#include "AddToProjectConfig.h"
+#include "FeaturedClasses.inl"
+#include "Classes/EditorStyleSettings.h"
+
+#include "Subsystems/PanelExtensionSubsystem.h"	// SExtensionPanel
+#include "Subsystems/AssetEditorSubsystem.h"
+
+#include "Framework/Notifications/NotificationManager.h"
+#include "Widgets/Notifications/SNotificationList.h"
+#include "BPVariableDragDropAction.h"
+
+#include "DragAndDrop/AssetDragDropOp.h"
+#include "K2Node_Variable.h"		// Nodes for FSubobjectRowDragDropOp
+#include "K2Node_VariableGet.h"
+#include "AssetSelection.h"
+#include "Settings/EditorProjectSettings.h"
+
+extern UNREALED_API UEditorEngine* GEditor;
+
+#define LOCTEXT_NAMESPACE "SSubobjectEditor"
+
+DEFINE_LOG_CATEGORY_STATIC(LogSubobjectEditor, Log, All);
+
+static const FName SubobjectTree_ColumnName_ComponentClass("ComponentClass");
+static const FName SubobjectTree_ColumnName_Asset("Asset");
+static const FName SubobjectTree_ColumnName_Mobility("Mobility"); 
+static const FName SubobjectTree_ContextMenuName("Kismet.SubobjectEditorContextMenu");
+
+///////////////////////////////////////////////////////////
+// Util Functions
+namespace SubobjectEditorUtils
+{	
+	static bool ContainsChildActorSubtreeNode(const TArray<FSubobjectEditorTreeNodePtrType>& InNodePtrs)
+	{
+		for (FSubobjectEditorTreeNodePtrType NodePtr : InNodePtrs)
+		{
+			const FSubobjectData* Data = NodePtr ? NodePtr->GetDataSource() : nullptr;
+			if (Data && Data->IsChildActor())
+			{
+				return true;
+			}
+		}
+	
+		return false;
+	}
+	
+	/** Returns whether child actor tree view expansion is enabled in project settings */
+	static bool IsChildActorTreeViewExpansionEnabled()
+	{
+		const UBlueprintEditorProjectSettings* EditorProjectSettings = GetDefault<UBlueprintEditorProjectSettings>();
+		return EditorProjectSettings->bEnableChildActorExpansionInTreeView;
+	}
+
+	static void OnSetChildActorTreeViewVisualizationMode(const UChildActorComponent* InChildActorComponent, EChildActorComponentTreeViewVisualizationMode InMode, TWeakPtr<SSubobjectEditor> InWeakSubobjectEditorPtr)
+	{
+		if (!InChildActorComponent)
+		{
+			return;
+		}
+
+		// #TODO_BH Get rid of const cast
+		const_cast<UChildActorComponent*>(InChildActorComponent)->SetEditorTreeViewVisualizationMode(InMode);
+
+		TSharedPtr<SSubobjectEditor> SubobjectEditor = InWeakSubobjectEditorPtr.Pin();
+		if (SubobjectEditor.IsValid())
+		{
+			SubobjectEditor->UpdateTree();
+		}
+	}
+	
+	static void FillChildActorContextMenuOptions(UToolMenu* Menu, const FSubobjectEditorTreeNodePtrType InNodePtr)
+	{
+		if (!IsChildActorTreeViewExpansionEnabled())
+		{
+			return;
+		}
+
+		if (!InNodePtr || !InNodePtr->IsChildSubtreeNode())
+		{
+			return;
+		}
+
+		check(InNodePtr.IsValid());
+
+		const UChildActorComponent* ChildActorComponent = Cast<UChildActorComponent>(InNodePtr->GetComponentTemplate());
+		if (!ChildActorComponent)
+		{
+			return;
+		}
+
+		FToolMenuSection& Section = Menu->AddSection("ChildActor", LOCTEXT("ChildActorHeading", "Child Actor"));
+		{
+			TWeakPtr<SSubobjectEditor> WeakEditorPtr;
+			if (USSCSEditorMenuContext* MenuContext = Menu->FindContext<USSCSEditorMenuContext>())
+			{
+				WeakEditorPtr = MenuContext->SubobjectEditor;
+			}
+
+			// #TODO_BH Why does this give me a linker error???????????
+			// Section.AddMenuEntry(
+   //              "SetChildActorOnlyMode",
+   //              LOCTEXT("SetChildActorOnlyMode_Label", "Switch to Child Actor Only Mode"),
+   //              LOCTEXT("SetChildActorOnlyMode_ToolTip", "Visualize this child actor's template/instance subtree in place of its parent component node."),
+   //              FSlateIcon(),
+   //              FUIAction(
+   //                  FExecuteAction::CreateStatic(
+	  //                   &SubobjectEditorUtils::OnSetChildActorTreeViewVisualizationMode, 
+	  //                   ChildActorComponent,
+	  //                   EChildActorComponentTreeViewVisualizationMode::ChildActorOnly,
+	  //                   WeakEditorPtr),
+   //                  FCanExecuteAction()
+   //              )
+   //          );
+		}
+	}
+}
+
+/////////////////////////////////////////////////////
+// FSubobjectEditorTreeNode
+
+FSubobjectEditorTreeNode::FSubobjectEditorTreeNode(const FSubobjectDataHandle& SubobjectDataSource, bool InbIsSeperator /*= false*/)
+	: DataHandle(SubobjectDataSource)
+	, FilterFlags((uint8)EFilteredState::Unknown)
+	, bIsSeperator(InbIsSeperator)
+{
+	check(IsValid());
+}
+
+FString FSubobjectEditorTreeNode::GetDisplayString() const
+{
+	FSubobjectData* Data = GetDataSource();
+	return Data ? Data->GetDisplayString() : TEXT("Unknown");
+}
+
+bool FSubobjectEditorTreeNode::IsComponentNode() const
+{
+	FSubobjectData* Data = GetDataSource();
+	return Data ? Data->GetComponentTemplate() != nullptr : false;
+}
+
+bool FSubobjectEditorTreeNode::IsNativeComponent() const
+{
+	FSubobjectData* Data = GetDataSource();
+	return Data ? Data->IsNativeComponent() : false;
+}
+
+bool FSubobjectEditorTreeNode::IsRootActorNode() const
+{
+	// If there is no parent node pointer then this has to be top of the tree
+	return ParentNodePtr == nullptr;
+}
+
+FName FSubobjectEditorTreeNode::GetVariableName() const
+{
+	const FSubobjectData* const Data = GetDataSource();
+	return Data->GetVariableName();
+}
+
+bool FSubobjectEditorTreeNode::CanReparent() const
+{
+	const FSubobjectData* const Data = GetDataSource();
+	return Data->CanReparent();
+}
+
+const UActorComponent* FSubobjectEditorTreeNode::GetComponentTemplate() const
+{
+	const FSubobjectData* Data = GetDataSource();
+	return Data ? Data->GetComponentTemplate() : nullptr;
+}
+
+const UObject* FSubobjectEditorTreeNode::GetObject(bool bEvenIfPendingKill /* = false */) const
+{
+	const FSubobjectData* Data = GetDataSource();
+	return Data ? Data->GetObject() : nullptr;
+}
+
+FSubobjectData* FSubobjectEditorTreeNode::GetDataSource() const
+{
+	return DataHandle.GetSharedDataPtr().Get();
+}
+
+bool FSubobjectEditorTreeNode::IsChildSubtreeNode() const
+{
+	const FSubobjectData* Data = GetDataSource();
+	return Data ? Data->IsChildActor() : false;
+}
+
+bool FSubobjectEditorTreeNode::IsAttachedTo(FSubobjectEditorTreeNodePtrType InNodePtr) const
+{
+	FSubobjectEditorTreeNodePtrType TestParentPtr = ParentNodePtr;
+	while (TestParentPtr.IsValid())
+	{
+		if (TestParentPtr == InNodePtr)
+		{
+			return true;
+		}
+
+		TestParentPtr = TestParentPtr->ParentNodePtr;
+	}
+
+	return false;
+}
+
+bool FSubobjectEditorTreeNode::IsDirectlyAttachedTo(FSubobjectEditorTreeNodePtrType InNodePtr) const
+{
+	return ParentNodePtr == InNodePtr;
+}
+
+bool FSubobjectEditorTreeNode::CanDelete() const
+{
+	if(const FSubobjectData* Data = GetDataSource())
+	{
+		return Data->CanDelete();
+	}
+	return false;
+}
+
+bool FSubobjectEditorTreeNode::CanRename() const
+{
+	if(const FSubobjectData* Data = GetDataSource())
+	{
+		return Data->CanRename();
+	}
+	return false;
+}
+
+void FSubobjectEditorTreeNode::AddChild(FSubobjectEditorTreeNodePtrType AttachToPtr)
+{
+	check(AttachToPtr->IsValid());
+	
+	// Ensure the node is not already parented elsewhere
+	if (AttachToPtr->GetParent())
+	{
+		AttachToPtr->GetParent()->RemoveChild(AttachToPtr);
+	}
+	
+	Children.AddUnique(AttachToPtr);
+	AttachToPtr->ParentNodePtr = AsShared();
+}
+
+void FSubobjectEditorTreeNode::RemoveChild(FSubobjectEditorTreeNodePtrType InChildNodePtr)
+{
+	Children.Remove(InChildNodePtr);
+	InChildNodePtr->ParentNodePtr.Reset();
+}
+
+FSubobjectEditorTreeNodePtrType FSubobjectEditorTreeNode::FindChild(const FSubobjectDataHandle& InHandle)
+{
+	for (const FSubobjectEditorTreeNodePtrType& Child : Children)
+	{
+		if (Child->GetDataHandle() == InHandle)
+		{
+			return Child;
+		}
+	}
+	return FSubobjectEditorTreeNodePtrType();
+}
+
+bool FSubobjectEditorTreeNode::IsFlaggedForFiltration() const
+{
+	return ensureMsgf(FilterFlags != EFilteredState::Unknown, TEXT("Querying a bad filtration state.")) ?
+		(FilterFlags & EFilteredState::FilteredInMask) == 0 : false;
+}
+
+void FSubobjectEditorTreeNode::SetCachedFilterState(bool bMatchesFilter, bool bUpdateParent)
+{
+	bool bFlagsChanged = false;
+	if ((FilterFlags & EFilteredState::Unknown) == EFilteredState::Unknown)
+	{
+		FilterFlags = 0x00;
+		bFlagsChanged = true;
+	}
+
+	if (bMatchesFilter)
+	{
+		bFlagsChanged |= (FilterFlags & EFilteredState::MatchesFilter) == 0;
+		FilterFlags |= EFilteredState::MatchesFilter;
+	}
+	else
+	{
+		bFlagsChanged |= (FilterFlags & EFilteredState::MatchesFilter) != 0;
+		FilterFlags &= ~EFilteredState::MatchesFilter;
+	}
+
+	const bool bHadChildMatch = (FilterFlags & EFilteredState::ChildMatches) != 0;
+
+	// refresh the cached child state (don't update the parent, we'll do that below if it's needed)
+	RefreshCachedChildFilterState(/*bUpdateParent =*/false);
+
+	bFlagsChanged |= bHadChildMatch != ((FilterFlags & EFilteredState::ChildMatches) != 0);
+	if (bUpdateParent && bFlagsChanged)
+	{
+		ApplyFilteredStateToParent();
+	}
+}
+
+void FSubobjectEditorTreeNode::ApplyFilteredStateToParent()
+{
+	FSubobjectEditorTreeNode* Child = this;
+	while (Child->ParentNodePtr.IsValid())
+	{
+		FSubobjectEditorTreeNode* Parent = Child->ParentNodePtr.Get();
+
+		if (!IsFlaggedForFiltration())
+		{
+			if ((Parent->FilterFlags & EFilteredState::ChildMatches) == 0)
+			{
+				Parent->FilterFlags |= EFilteredState::ChildMatches;
+			}
+			else
+			{
+				// all parents from here on up should have the flag
+				break;
+			}
+		}
+		// have to see if this was the only child contributing to this flag
+		else if (Parent->FilterFlags & EFilteredState::ChildMatches)
+		{
+			Parent->FilterFlags &= ~EFilteredState::ChildMatches;
+			for (const FSubobjectEditorTreeNodePtrType& Sibling : Parent->Children)
+			{
+				if (Sibling.Get() == Child)
+				{
+					continue;
+				}
+
+				if (Sibling->FilterFlags & EFilteredState::FilteredInMask)
+				{
+					Parent->FilterFlags |= EFilteredState::ChildMatches;
+					break;
+				}
+			}
+
+			if (Parent->FilterFlags & EFilteredState::ChildMatches)
+			{
+				// another child added the flag back
+				break;
+			}
+		}
+		Child = Parent;
+	}
+}
+
+void FSubobjectEditorTreeNode::RefreshCachedChildFilterState(bool bUpdateParent)
+{
+	const bool bContainedMatch = !IsFlaggedForFiltration();
+
+	FilterFlags &= ~EFilteredState::ChildMatches;
+	for (FSubobjectEditorTreeNodePtrType Child : Children)
+	{
+		// Separator nodes should not contribute to child matches for the parent nodes
+		if (Child->IsSeperator())
+		{
+			continue;
+		}
+
+		if (!Child->IsFlaggedForFiltration())
+		{
+			FilterFlags |= EFilteredState::ChildMatches;
+			break;
+		}
+	}
+	const bool bContainsMatch = !IsFlaggedForFiltration();
+
+	const bool bStateChange = bContainedMatch != bContainsMatch;
+	if (bUpdateParent && bStateChange)
+	{
+		ApplyFilteredStateToParent();
+	}
+}
+
+bool FSubobjectEditorTreeNode::RefreshFilteredState(const UClass* InFilterType, const TArray<FString>& InFilterTerms, bool bRecursive)
+{
+	bool bHasAnyVisibleChildren = false;
+	if (bRecursive)
+	{
+		for (FSubobjectEditorTreeNodePtrType Child : GetChildren())
+		{
+			bHasAnyVisibleChildren |= Child->RefreshFilteredState(InFilterType, InFilterTerms, bRecursive);
+		}
+	}
+
+	// Don't check a root actor node - it doesn't have a valid variable name. Let it recache based on children and hide itself based on their filter states.
+	if (IsRootActorNode())
+	{
+		SetCachedFilterState(bHasAnyVisibleChildren, /*bUpdateParent =*/!bRecursive);
+		return bHasAnyVisibleChildren;
+	}
+
+	bool bIsFilteredOut = InFilterType && !MatchesFilterType(InFilterType);
+	if (!bIsFilteredOut && !IsSeperator())
+	{
+		FString DisplayStr = GetDisplayString();
+		for (const FString& FilterTerm : InFilterTerms)
+		{
+			if (!DisplayStr.Contains(FilterTerm))
+			{
+				bIsFilteredOut = true;
+			}
+		}
+	}
+
+	// if we're not recursing, then assume this is for a new node and we need to update the parent
+	// otherwise, assume the parent was hit as part of the recursion
+	const bool bUpdateParent = !bRecursive;
+	SetCachedFilterState(!bIsFilteredOut, bUpdateParent);
+	return !bIsFilteredOut;
+}
+
+bool FSubobjectEditorTreeNode::MatchesFilterType(const UClass* InFilterType) const
+{
+	if(IsComponentNode())
+	{
+		if (const UActorComponent* ComponentObject = GetComponentTemplate())
+		{
+			const UClass* ComponentClass = ComponentObject->GetClass();
+			check(ComponentClass);
+
+			if (ComponentClass->IsChildOf(InFilterType))
+			{
+				return true;
+			}
+		}
+		return false;
+	}
+	
+	// All other nodes will pass the type filter by default.
+	return true;
+}
+
+void FSubobjectEditorTreeNode::SetOngoingCreateTransaction(TUniquePtr<FScopedTransaction> InTransaction)
+{
+	OngoingCreateTransaction = MoveTemp(InTransaction);
+}
+
+void FSubobjectEditorTreeNode::CloseOngoingCreateTransaction()
+{
+	OngoingCreateTransaction.Reset();
+}
+
+//////////////////////////////////////////////////////////////////////////
+// FSubobjectRowDragDropOp - The drag-drop operation triggered when dragging a row in the components tree
+
+class FSubobjectRowDragDropOp : public FKismetVariableDragDropAction
+{
+public:
+	DRAG_DROP_OPERATOR_TYPE(FSubobjectRowDragDropOp, FKismetVariableDragDropAction)
+
+	/** Available drop actions */
+	enum EDropActionType
+	{
+		DropAction_None,
+		DropAction_AttachTo,
+		DropAction_DetachFrom,
+		DropAction_MakeNewRoot,
+		DropAction_AttachToOrMakeNewRoot
+	};
+
+	// FGraphEditorDragDropAction interface
+	virtual void HoverTargetChanged() override;
+	virtual FReply DroppedOnNode(FVector2D ScreenPosition, FVector2D GraphPosition) override;
+	virtual FReply DroppedOnPanel(const TSharedRef< class SWidget >& Panel, FVector2D ScreenPosition, FVector2D GraphPosition, UEdGraph& Graph) override;
+	// End of FGraphEditorDragDropAction
+
+	/** Node(s) that we started the drag from */
+	TArray<FSubobjectEditorTreeNodePtrType> SourceNodes;
+
+	/** The type of drop action that's pending while dragging */
+	EDropActionType PendingDropAction;
+
+	static TSharedRef<FSubobjectRowDragDropOp> New(FName InVariableName, UStruct* InVariableSource, FNodeCreationAnalytic AnalyticCallback);
+};
+
+TSharedRef<FSubobjectRowDragDropOp> FSubobjectRowDragDropOp::New(FName InVariableName, UStruct* InVariableSource, FNodeCreationAnalytic AnalyticCallback)
+{
+	TSharedPtr<FSubobjectRowDragDropOp> Operation = MakeShareable(new FSubobjectRowDragDropOp);
+	Operation->VariableName = InVariableName;
+	Operation->VariableSource = InVariableSource;
+	Operation->AnalyticCallback = AnalyticCallback;
+	Operation->Construct();
+	return Operation.ToSharedRef();
+}
+
+void FSubobjectRowDragDropOp::HoverTargetChanged()
+{
+	bool bHoverHandled = false;
+
+	FSlateColor IconTint = FLinearColor::White;
+	const FSlateBrush* ErrorSymbol = FEditorStyle::GetBrush(TEXT("Graph.ConnectorFeedback.Error"));
+
+	if(SourceNodes.Num() > 1)
+	{
+		// Display an error message if attempting to drop multiple source items onto a node
+		UK2Node_Variable* VarNodeUnderCursor = Cast<UK2Node_Variable>(GetHoveredNode());
+		if (VarNodeUnderCursor != nullptr)
+		{
+			// Icon/text to draw on tooltip
+			FText Message = LOCTEXT("InvalidMultiDropTarget", "Cannot replace node with multiple nodes");
+			SetSimpleFeedbackMessage(ErrorSymbol, IconTint, Message);
+
+			bHoverHandled = true;
+		}
+	}
+
+	if (!bHoverHandled)
+	{
+		if (SubobjectEditorUtils::ContainsChildActorSubtreeNode(SourceNodes))
+		{
+			// @todo - Add support for drag/drop of child actor template components to create variable get nodes in a local Blueprint event/function graph
+			FText Message = LOCTEXT("ChildActorDragDropAddVariableNode_Unsupported", "This operation is not currently supported for one or more of the selected components.");
+			SetSimpleFeedbackMessage(ErrorSymbol, IconTint, Message);
+		}
+		else if (FProperty* VariableProperty = GetVariableProperty())
+		{
+			const FSlateBrush* PrimarySymbol;
+			const FSlateBrush* SecondarySymbol;
+			FSlateColor PrimaryColor;
+			FSlateColor SecondaryColor;
+			GetDefaultStatusSymbol(/*out*/ PrimarySymbol, /*out*/ PrimaryColor, /*out*/ SecondarySymbol, /*out*/ SecondaryColor);
+
+			// Create feedback message with the function name.
+			SetSimpleFeedbackMessage(PrimarySymbol, PrimaryColor, VariableProperty->GetDisplayNameText(), SecondarySymbol, SecondaryColor);
+		}
+		else
+		{
+			FText Message = LOCTEXT("CannotFindProperty", "Cannot find corresponding variable (make sure component has been assigned to one)");
+			SetSimpleFeedbackMessage(ErrorSymbol, IconTint, Message);
+		}
+		bHoverHandled = true;
+	}
+
+	if(!bHoverHandled)
+	{
+		FKismetVariableDragDropAction::HoverTargetChanged();
+	}
+}
+
+FReply FSubobjectRowDragDropOp::DroppedOnNode(FVector2D ScreenPosition, FVector2D GraphPosition)
+{
+	// Only allow dropping on another node if there is only a single source item
+	if(SourceNodes.Num() == 1)
+	{
+		FKismetVariableDragDropAction::DroppedOnNode(ScreenPosition, GraphPosition);
+	}
+	return FReply::Handled();
+}
+
+FReply FSubobjectRowDragDropOp::DroppedOnPanel(const TSharedRef<class SWidget>& Panel, FVector2D ScreenPosition, FVector2D GraphPosition, UEdGraph& Graph)
+{
+	const FScopedTransaction Transaction(LOCTEXT("SSubobjectEditorAddMultipleNodes", "Add Component Nodes"));
+
+	TArray<UK2Node_VariableGet*> OriginalVariableNodes;
+	Graph.GetNodesOfClass<UK2Node_VariableGet>(OriginalVariableNodes);
+
+	// Add source items to the graph in turn
+	for (FSubobjectEditorTreeNodePtrType& SourceNode : SourceNodes)
+	{
+		FSubobjectData* Data = SourceNode->GetDataSource();
+		if(!Data)
+		{
+			continue;
+		}
+		
+		VariableName = Data->GetVariableName();
+		FKismetVariableDragDropAction::DroppedOnPanel(Panel, ScreenPosition, GraphPosition, Graph);
+
+		GraphPosition.Y += 50.0f;
+	}
+
+	TArray<UK2Node_VariableGet*> ResultVariableNodes;
+	Graph.GetNodesOfClass<UK2Node_VariableGet>(ResultVariableNodes);
+
+	if (ResultVariableNodes.Num() - OriginalVariableNodes.Num() > 1)
+	{
+		TSet<const UEdGraphNode*> NodeSelection;
+
+		// Because there is more than one new node, lets grab all the nodes at the bottom of the list and add them to a set for selection
+		for (int32 NodeIdx = ResultVariableNodes.Num() - 1; NodeIdx >= OriginalVariableNodes.Num(); --NodeIdx)
+		{
+			NodeSelection.Add(ResultVariableNodes[NodeIdx]);
+		}
+		Graph.SelectNodeSet(NodeSelection);
+	}
+	return FReply::Handled();
+}
+
+/////////////////////////////////////////////////////
+// SSubobjectEditorDragDropTree
+
+void SSubobjectEditorDragDropTree::Construct(const FArguments& InArgs)
+{
+	STreeView<FSubobjectEditorTreeNodePtrType>::FArguments BaseArgs;
+	BaseArgs.OnGenerateRow(InArgs._OnGenerateRow)
+		.OnItemScrolledIntoView(InArgs._OnItemScrolledIntoView)
+		.OnGetChildren(InArgs._OnGetChildren)
+		.OnSetExpansionRecursive(InArgs._OnSetExpansionRecursive)
+		.TreeItemsSource(InArgs._TreeItemsSource)
+		.ItemHeight(InArgs._ItemHeight)
+		.OnContextMenuOpening(InArgs._OnContextMenuOpening)
+		.OnMouseButtonDoubleClick(InArgs._OnMouseButtonDoubleClick)
+		.OnSelectionChanged(InArgs._OnSelectionChanged)
+		.SelectionMode(InArgs._SelectionMode)
+		.HeaderRow(InArgs._HeaderRow)
+		.ClearSelectionOnClick(InArgs._ClearSelectionOnClick)
+		.ExternalScrollbar(InArgs._ExternalScrollbar)
+		.OnEnteredBadState(InArgs._OnTableViewBadState)
+		.HighlightParentNodesForSelection(true);
+
+	STreeView<FSubobjectEditorTreeNodePtrType>::Construct(BaseArgs);
+}
+
+FReply SSubobjectEditorDragDropTree::OnDragOver(const FGeometry& MyGeometry, const FDragDropEvent& DragDropEvent)
+{
+	FReply Handled = FReply::Unhandled();
+
+	if (SubobjectEditor != nullptr)
+	{
+		TSharedPtr<FDragDropOperation> Operation = DragDropEvent.GetOperation();
+		if (Operation.IsValid() && (Operation->IsOfType<FExternalDragOperation>() || Operation->IsOfType<FAssetDragDropOp>()))
+		{
+			Handled = AssetUtil::CanHandleAssetDrag(DragDropEvent);
+
+			if (!Handled.IsEventHandled())
+			{
+				if (Operation->IsOfType<FAssetDragDropOp>())
+				{
+					const TSharedPtr<FAssetDragDropOp> AssetDragDropOp = StaticCastSharedPtr<FAssetDragDropOp>(Operation);
+
+					for (const FAssetData& AssetData : AssetDragDropOp->GetAssets())
+					{
+						if (UClass* AssetClass = AssetData.GetClass())
+						{
+							if (AssetClass->IsChildOf(UClass::StaticClass()))
+							{
+								Handled = FReply::Handled();
+								break;
+							}
+						}
+					}
+				}
+			}
+		}
+	}
+
+	return Handled;
+}
+
+FReply SSubobjectEditorDragDropTree::OnDrop(const FGeometry& MyGeometry, const FDragDropEvent& DragDropEvent)
+{
+	if (SubobjectEditor != nullptr)
+	{
+		return SubobjectEditor->TryHandleAssetDragDropOperation(DragDropEvent);
+	}
+	else
+	{
+		return FReply::Unhandled();
+	}
+}
+
+/////////////////////////////////////////////////////
+// SSubobject_RowWidget
+
+void SSubobject_RowWidget::Construct(const FArguments& InArgs, TWeakPtr<SSubobjectEditor> InEditor, FSubobjectEditorTreeNodePtrType InNodePtr, TSharedPtr<STableViewBase> InOwnerTableView)
+{
+	check(InNodePtr.IsValid());
+	SubobjectEditor = InEditor;
+	SubobjectPtr = InNodePtr;
+
+	bool bIsSeparator = InNodePtr->IsSeperator();
+
+	FSuperRowType::FArguments Args = FSuperRowType::FArguments()
+		.Style(bIsSeparator ?
+			&FEditorStyle::Get().GetWidgetStyle<FTableRowStyle>("TableView.NoHoverTableRow") :
+			&FEditorStyle::Get().GetWidgetStyle<FTableRowStyle>("SceneOutliner.TableViewRow")) //@todo create editor style for the SCS tree
+		.Padding(FMargin(0.f, 4.f, 0.f, 4.f))
+		.ShowSelection(!bIsSeparator)
+		.OnDragDetected(this, &SSubobject_RowWidget::HandleOnDragDetected)
+		.OnDragEnter(this, &SSubobject_RowWidget::HandleOnDragEnter)
+		.OnDragLeave(this, &SSubobject_RowWidget::HandleOnDragLeave)
+		.OnCanAcceptDrop(this, &SSubobject_RowWidget::HandleOnCanAcceptDrop)
+		.OnAcceptDrop(this, &SSubobject_RowWidget::HandleOnAcceptDrop);
+
+	SMultiColumnTableRow<FSubobjectEditorTreeNodePtrType>::Construct(Args, InOwnerTableView.ToSharedRef());
+}
+
+BEGIN_SLATE_FUNCTION_BUILD_OPTIMIZATION
+TSharedRef<SWidget> SSubobject_RowWidget::GenerateWidgetForColumn(const FName& ColumnName)
+{
+	FSubobjectEditorTreeNodePtrType Node = GetSubobjectPtr();
+	check(Node);
+	
+	if (Node->IsSeperator())
+	{
+		static const FLinearColor Debug(1.0f, 0.0f, 0.0f);
+		return SNew(SSpacer);
+		// return SNew(SBox)
+		// 	.Padding(1.f)
+		// 	[
+		// 		SNew(SBorder)
+		// 		.Padding(FEditorStyle::GetMargin(TEXT("Menu.Separator.Padding")))
+		// 		.BorderImage(FEditorStyle::GetBrush(TEXT("Menu.Separator")))
+		// 		.ColorAndOpacity(Debug)
+		// 	];
+	}
+
+	FSubobjectData* Data = Node->GetDataSource();
+
+	if (ColumnName == SubobjectTree_ColumnName_ComponentClass)
+	{
+		InlineWidget =
+			SNew(SInlineEditableTextBlock)
+			.Text(this, &SSubobject_RowWidget::GetNameLabel)
+			.OnVerifyTextChanged(this, &SSubobject_RowWidget::OnNameTextVerifyChanged)
+			.OnTextCommitted(this, &SSubobject_RowWidget::OnNameTextCommit)
+			.IsSelected(this, &SSubobject_RowWidget::IsSelectedExclusively)
+			.IsReadOnly(this, &SSubobject_RowWidget::IsReadOnly);
+
+		Node->SetRenameRequestedDelegate(FSubobjectEditorTreeNode::FOnRenameRequested::CreateSP(InlineWidget.Get(), &SInlineEditableTextBlock::EnterEditingMode));
+		
+		const bool IsRootActorNode = Data->IsRootActor();
+		
+		TSharedRef<SHorizontalBox> Contents = 
+		SNew(SHorizontalBox)
+			.ToolTip(CreateToolTipWidget())
+			+ SHorizontalBox::Slot()
+			.Padding(FMargin(IsRootActorNode ? 0.f : 4.f, 0.f, 0.f, 0.f))
+			.AutoWidth()
+			.VAlign(VAlign_Center)
+			[
+				SNew(SExpanderArrow, SharedThis(this))
+				.Visibility(IsRootActorNode ? EVisibility::Collapsed : EVisibility::Visible)
+			]
+		+ SHorizontalBox::Slot()
+			.Padding(FMargin(IsRootActorNode ? 4.f : 0.f, 0.f, 0.f, 0.f))
+			.AutoWidth()
+			.VAlign(VAlign_Center)
+			[
+				SNew(SImage)
+				.Image(GetIconBrush())
+				.ColorAndOpacity(this, &SSubobject_RowWidget::GetColorTintForIcon)
+			]
+		+ SHorizontalBox::Slot()
+			.VAlign(VAlign_Center)
+			.HAlign(HAlign_Left)
+			.AutoWidth()
+			.Padding(6.f, 0.f, 0.f, 0.f)
+			[
+				InlineWidget.ToSharedRef()
+			]
+		+ SHorizontalBox::Slot()
+	        .HAlign(HAlign_Left)
+	        .VAlign(VAlign_Center)
+	        .Padding(4.f, 0.f, 0.f, 0.f)
+	        [
+	            SNew(STextBlock)
+	            .Text(this, &SSubobject_RowWidget::GetObjectContextText)
+	            .ColorAndOpacity(FSlateColor::UseForeground())
+	        ];
+				
+		return Contents;
+	}
+	else if (ColumnName == SubobjectTree_ColumnName_Asset)
+	{
+		return
+			SNew(SHorizontalBox)
+			+ SHorizontalBox::Slot()
+			.VAlign(VAlign_Center)
+			[
+				SNew(STextBlock)
+				.Visibility(this, &SSubobject_RowWidget::GetAssetVisibility)
+				.Text(this, &SSubobject_RowWidget::GetAssetName)
+				.ToolTipText(this, &SSubobject_RowWidget::GetAssetPath)
+			];
+	}
+	else if (ColumnName == SubobjectTree_ColumnName_Mobility)
+	{
+		if (Data && Data->GetComponentTemplate())
+		{
+			TSharedPtr<SToolTip> MobilityTooltip = SNew(SToolTip)
+				.Text(this, &SSubobject_RowWidget::GetMobilityToolTipText);
+
+			return SNew(SHorizontalBox)
+				.ToolTip(MobilityTooltip)
+				.Visibility(EVisibility::Visible) // so we still get tooltip text for an empty SHorizontalBox
+				+ SHorizontalBox::Slot()
+				.VAlign(VAlign_Center)
+				.FillWidth(1.0f)
+				[
+					SNew(SImage)
+					.Image(this, &SSubobject_RowWidget::GetMobilityIconImage)
+					.ToolTip(MobilityTooltip)
+				];
+		}
+		else
+		{
+			return SNew(SSpacer);
+		}
+	}
+
+	return SNew(STextBlock)
+		.Text(LOCTEXT("UnknownColumn", "Unknown Column"));
+}
+
+FText SSubobject_RowWidget::GetNameLabel() const
+{
+	FSubobjectEditorTreeNodePtrType NodePtr = GetSubobjectPtr();
+	FSubobjectData* Data = NodePtr ? NodePtr->GetDataSource() : nullptr;
+	if(Data)
+	{
+		// NOTE: Whatever this returns also becomes the variable name
+		return FText::FromString(NodePtr->GetDisplayString());
+	}
+	return LOCTEXT("UnknownDataSource", "Unknown Name Label");
+}
+
+FText SSubobject_RowWidget::GetAssetPath() const
+{
+	FSubobjectEditorTreeNodePtrType NodePtr = GetSubobjectPtr();
+	FSubobjectData* Data = NodePtr->GetDataSource();
+
+	return Data ? Data->GetAssetPath() : FText::GetEmpty();
+}
+
+FText SSubobject_RowWidget::GetAssetName() const
+{
+	FSubobjectEditorTreeNodePtrType NodePtr = GetSubobjectPtr();
+	if(FSubobjectData* Data = NodePtr->GetDataSource())
+	{
+		return Data->GetAssetName();
+
+	}
+	return FText::GetEmpty();
+}
+
+EVisibility SSubobject_RowWidget::GetAssetVisibility() const
+{	
+	FSubobjectEditorTreeNodePtrType NodePtr = GetSubobjectPtr();
+	FSubobjectData* Data = NodePtr ? NodePtr->GetDataSource() : nullptr;
+
+	if (Data && Data->IsAssetVisible())
+	{
+		return EVisibility::Visible;
+	}
+
+	return EVisibility::Hidden;
+}
+
+const FSlateBrush* SSubobject_RowWidget::GetIconBrush() const
+{
+	const FSlateBrush* ComponentIcon = FEditorStyle::GetBrush("SCS.NativeComponent");
+
+	if (FSubobjectEditorTreeNodePtrType NodePtr = GetSubobjectPtr())
+	{
+		if (const FSubobjectData* Data = NodePtr->GetDataSource())
+		{
+			if (const AActor* Actor = Data->GetObject<AActor>())
+			{
+				return FClassIconFinder::FindIconForActor(Actor);
+			}
+			else if (const UActorComponent* ComponentTemplate = Data->GetComponentTemplate())
+			{
+				ComponentIcon = FSlateIconFinder::FindIconBrushForClass(ComponentTemplate->GetClass(), TEXT("SCS.Component"));
+			}
+		}
+	}
+
+	return ComponentIcon;
+}
+
+FSlateColor SSubobject_RowWidget::GetColorTintForIcon() const
+{
+	return SubobjectEditor.Pin()->GetColorTintForIcon(GetSubobjectPtr());
+}
+
+ESelectionMode::Type SSubobject_RowWidget::GetSelectionMode() const
+{
+	FSubobjectEditorTreeNodePtrType NodePtr = GetSubobjectPtr();
+	if (NodePtr && NodePtr->IsSeperator())
+	{
+		return ESelectionMode::None;
+	}
+	
+	return SMultiColumnTableRow<FSubobjectEditorTreeNodePtrType>::GetSelectionMode();
+}
+
+void SSubobject_RowWidget::HandleOnDragEnter(const FDragDropEvent& DragDropEvent)
+{
+	TSharedPtr<FDragDropOperation> Operation = DragDropEvent.GetOperation();
+	if (!Operation.IsValid())
+	{
+		return;
+	}
+	TSharedPtr<FSubobjectRowDragDropOp> DragRowOp = DragDropEvent.GetOperationAs<FSubobjectRowDragDropOp>();
+
+	if (DragRowOp.IsValid())
+	{
+		check(SubobjectEditor.IsValid());
+		
+		FText Message;
+		FSlateColor IconColor = FLinearColor::White;
+
+		for (const FSubobjectEditorTreeNodePtrType& SelectedNodePtr : DragRowOp->SourceNodes)
+		{
+			const FSubobjectData* Data = SelectedNodePtr->GetDataSource();
+			if (!Data->CanReparent())
+			{
+				// We set the tooltip text here because it won't change across entry/leave events
+				if (DragRowOp->SourceNodes.Num() == 1)
+				{
+					if (SelectedNodePtr->IsChildSubtreeNode())
+					{
+						Message = LOCTEXT("DropActionToolTip_Error_CannotReparent_ChildActorSubTreeNodes", "The selected component is part of a child actor template and cannot be reordered here.");
+					}
+					else if (!Data->IsSceneComponent())
+					{
+						Message = LOCTEXT("DropActionToolTip_Error_CannotReparent_NotSceneComponent", "The selected component is not a scene component and cannot be attached to other components.");
+					}
+					else if (Data->IsInheritedComponent())
+					{
+						Message = LOCTEXT("DropActionToolTip_Error_CannotReparent_Inherited", "The selected component is inherited and cannot be reordered here.");
+					}
+					else
+					{
+						Message = LOCTEXT("DropActionToolTip_Error_CannotReparent", "The selected component cannot be moved.");
+					}
+				}
+				else
+				{
+					Message = LOCTEXT("DropActionToolTip_Error_CannotReparentMultiple", "One or more of the selected components cannot be attached.");
+				}
+				break;
+			}
+		}
+
+		if (Message.IsEmpty())
+		{
+			FSubobjectEditorTreeNodePtrType SceneRootNodePtr = SubobjectEditor.Pin()->GetSceneRootNode();
+
+			FSubobjectEditorTreeNodePtrType NodePtr = GetSubobjectPtr();
+			const FSubobjectData* NodeData = NodePtr->GetDataSource();
+
+			if (!NodePtr->IsComponentNode())
+			{
+				// Don't show a feedback message if over a node that makes no sense, such as a separator or the instance node
+				Message = LOCTEXT("DropActionToolTip_FriendlyError_DragToAComponent", "Drag to another component in order to attach to that component or become the root component.\nDrag to a Blueprint graph in order to drop a reference.");
+			}
+			else if (NodePtr->IsChildSubtreeNode())
+			{
+				// Can't drag onto components within a child actor node's subtree
+				Message = LOCTEXT("DropActionToolTip_Error_ChildActorSubTree", "Cannot attach to this component as it is part of a child actor template.");
+			}
+			
+			// Validate each selected node being dragged against the node that belongs to this row. Exit the loop if we have a valid tooltip OR a valid pending drop action once all nodes in the selection have been validated.
+			for (auto SourceNodeIter = DragRowOp->SourceNodes.CreateConstIterator(); SourceNodeIter && (Message.IsEmpty() || DragRowOp->PendingDropAction != FSubobjectRowDragDropOp::DropAction_None); ++SourceNodeIter)
+			{
+				FSubobjectEditorTreeNodePtrType DraggedNodePtr = *SourceNodeIter;
+				const FSubobjectData* const DraggedNodeData = DraggedNodePtr->GetDataSource();
+				check(DraggedNodePtr.IsValid());
+
+				// Reset the pending drop action each time through the loop
+				DragRowOp->PendingDropAction = FSubobjectRowDragDropOp::DropAction_None;
+
+				// Get the component template objects associated with each node
+				const USceneComponent* HoveredTemplate = Cast<USceneComponent>(NodePtr->GetComponentTemplate());
+				const USceneComponent* DraggedTemplate = Cast<USceneComponent>(DraggedNodeData->GetComponentTemplate());
+
+				if (DraggedNodePtr == NodePtr)
+				{
+					// Attempted to drag and drop onto self
+					if (DragRowOp->SourceNodes.Num() > 1)
+					{
+						Message = FText::Format(LOCTEXT("DropActionToolTip_Error_CannotAttachToSelfWithMultipleSelection", "Cannot attach the selected components here because it would result in {0} being attached to itself. Remove it from the selection and try again."), DraggedNodeData->GetDisplayName());
+					}
+					else
+					{
+						Message = FText::Format(LOCTEXT("DropActionToolTip_Error_CannotAttachToSelf", "Cannot attach {0} to itself."), DraggedNodeData->GetDisplayName());
+					}
+				}
+				else if (NodePtr->IsAttachedTo(DraggedNodePtr))
+				{
+					// Attempted to drop a parent onto a child
+					if (DragRowOp->SourceNodes.Num() > 1)
+					{
+						Message = FText::Format(LOCTEXT("DropActionToolTip_Error_CannotAttachToChildWithMultipleSelection", "Cannot attach the selected components here because it would result in {0} being attached to one of its children. Remove it from the selection and try again."), DraggedNodeData->GetDisplayName());
+					}
+					else
+					{
+						Message = FText::Format(LOCTEXT("DropActionToolTip_Error_CannotAttachToChild", "Cannot attach {0} to one of its children."), DraggedNodeData->GetDisplayName());
+					}
+				}
+				else if (HoveredTemplate == nullptr || DraggedTemplate == nullptr)
+				{
+					if (HoveredTemplate == nullptr)
+					{
+						// Can't attach to non-USceneComponent types
+						Message = LOCTEXT("DropActionToolTip_Error_NotAttachable_NotSceneComponent", "Cannot attach to this component as it is not a scene component.");
+					}
+					else
+					{
+						// Can't attach non-USceneComponent types
+						Message = LOCTEXT("DropActionToolTip_Error_NotAttachable", "Cannot attach to this component.");
+					}
+				}
+				else if (NodePtr == SceneRootNodePtr)
+				{
+					bool bCanMakeNewRoot = false;
+					bool bCanAttachToRoot = !DraggedNodePtr->IsDirectlyAttachedTo(NodePtr)
+						&& HoveredTemplate->CanAttachAsChild(DraggedTemplate, NAME_None)
+						&& DraggedTemplate->Mobility >= HoveredTemplate->Mobility
+						&& (!HoveredTemplate->IsEditorOnly() || DraggedTemplate->IsEditorOnly());
+
+					if (!NodePtr->CanReparent() && (!NodeData->IsDefaultSceneRoot() || NodeData->IsInheritedComponent()))
+					{
+						// Cannot make the dropped node the new root if we cannot reparent the current root
+						Message = LOCTEXT("DropActionToolTip_Error_CannotReparentRootNode", "The root component in this Blueprint is inherited and cannot be replaced.");
+					}
+					else if (DraggedTemplate->IsEditorOnly() && !HoveredTemplate->IsEditorOnly())
+					{
+						// can't have a new root that's editor-only (when children would be around in-game)
+						Message = LOCTEXT("DropActionToolTip_Error_CannotReparentEditorOnly", "Cannot re-parent game components under editor-only ones.");
+					}
+					else if (DraggedTemplate->Mobility > HoveredTemplate->Mobility)
+					{
+						// can't have a new root that's movable if the existing root is static or stationary
+						Message = LOCTEXT("DropActionToolTip_Error_CannotReparentNonMovable", "Cannot replace a non-movable scene root with a movable component.");
+					}
+					else if (DragRowOp->SourceNodes.Num() > 1)
+					{
+						Message = LOCTEXT("DropActionToolTip_Error_CannotAssignMultipleRootNodes", "Cannot replace the scene root with multiple components. Please select only a single component and try again.");
+					}
+					else
+					{
+						bCanMakeNewRoot = true;
+					}
+
+					if (bCanMakeNewRoot && bCanAttachToRoot)
+					{
+						// User can choose to either attach to the current root or make the dropped node the new root
+						Message = LOCTEXT("DropActionToolTip_AttachToOrMakeNewRoot", "Drop here to see available actions.");
+						DragRowOp->PendingDropAction = FSubobjectRowDragDropOp::DropAction_AttachToOrMakeNewRoot;
+					}
+					else if (SubobjectEditor.Pin()->CanMakeNewRootOnDrag(DraggedNodeData->GetBlueprint()))
+					{
+						if (bCanMakeNewRoot)
+						{
+							if (NodeData->IsDefaultSceneRoot())
+							{
+								// Only available action is to copy the dragged node to the other Blueprint and make it the new root
+								// Default root will be deleted
+								Message = FText::Format(LOCTEXT("DropActionToolTip_DropMakeNewRootNodeFromCopyAndDelete", "Drop here to copy {0} to a new variable and make it the new root. The default root will be deleted."), DraggedNodeData->GetDisplayName());
+							}
+							else
+							{
+								// Only available action is to copy the dragged node to the other Blueprint and make it the new root
+								Message = FText::Format(LOCTEXT("DropActionToolTip_DropMakeNewRootNodeFromCopy", "Drop here to copy {0} to a new variable and make it the new root."), DraggedNodeData->GetDisplayName());
+							}
+							DragRowOp->PendingDropAction = FSubobjectRowDragDropOp::DropAction_MakeNewRoot;
+						}
+						else if (bCanAttachToRoot)
+						{
+							// Only available action is to copy the dragged node(s) to the other Blueprint and attach it to the root
+							if (DragRowOp->SourceNodes.Num() > 1)
+							{
+								Message = FText::Format(LOCTEXT("DropActionToolTip_AttachComponentsToThisNodeFromCopyWithMultipleSelection", "Drop here to copy the selected components to new variables and attach them to {0}."), NodeData->GetDisplayName());
+							}
+							else
+							{
+								Message = FText::Format(LOCTEXT("DropActionToolTip_AttachToThisNodeFromCopy", "Drop here to copy {0} to a new variable and attach it to {1}."), DraggedNodeData->GetDisplayName(), NodeData->GetDisplayName());
+							}
+
+							DragRowOp->PendingDropAction = FSubobjectRowDragDropOp::DropAction_AttachTo;
+						}
+					}
+					else if (bCanMakeNewRoot)
+					{
+						if (NodeData->IsDefaultSceneRoot())
+						{
+							// Only available action is to make the dragged node the new root
+							// Default root will be deleted
+							Message = FText::Format(LOCTEXT("DropActionToolTip_DropMakeNewRootNodeAndDelete", "Drop here to make {0} the new root. The default root will be deleted."), DraggedNodeData->GetDisplayName());
+						}
+						else
+						{
+							// Only available action is to make the dragged node the new root
+							Message = FText::Format(LOCTEXT("DropActionToolTip_DropMakeNewRootNode", "Drop here to make {0} the new root."), DraggedNodeData->GetDisplayName());
+						}
+						DragRowOp->PendingDropAction = FSubobjectRowDragDropOp::DropAction_MakeNewRoot;
+					}
+					else if (bCanAttachToRoot)
+					{
+						// Only available action is to attach the dragged node(s) to the root
+						if (DragRowOp->SourceNodes.Num() > 1)
+						{
+							Message = FText::Format(LOCTEXT("DropActionToolTip_AttachToThisNodeWithMultipleSelection", "Drop here to attach the selected components to {0}."), NodeData->GetDisplayName());
+						}
+						else
+						{
+							Message = FText::Format(LOCTEXT("DropActionToolTip_AttachToThisNode", "Drop here to attach {0} to {1}."), DraggedNodeData->GetDisplayName(), NodeData->GetDisplayName());
+						}
+
+						DragRowOp->PendingDropAction = FSubobjectRowDragDropOp::DropAction_AttachTo;
+					}
+				}
+				else if (DraggedNodePtr->IsDirectlyAttachedTo(NodePtr)) // if dropped onto parent
+				{
+					// Detach the dropped node(s) from the current node and reattach to the root node
+					if (DragRowOp->SourceNodes.Num() > 1)
+					{
+						Message = FText::Format(LOCTEXT("DropActionToolTip_DetachFromThisNodeWithMultipleSelection", "Drop here to detach the selected components from {0}."), NodeData->GetDisplayName());
+					}
+					else
+					{
+						Message = FText::Format(LOCTEXT("DropActionToolTip_DetachFromThisNode", "Drop here to detach {0} from {1}."), DraggedNodeData->GetDisplayName(), NodeData->GetDisplayName());
+					}
+
+					DragRowOp->PendingDropAction = FSubobjectRowDragDropOp::DropAction_DetachFrom;
+				}
+				else if (!DraggedTemplate->IsEditorOnly() && HoveredTemplate->IsEditorOnly())
+				{
+					// can't have a game component child nested under an editor-only one
+					Message = LOCTEXT("DropActionToolTip_Error_CannotAttachToEditorOnly", "Cannot attach game components to editor-only ones.");
+				}
+				else if ((DraggedTemplate->Mobility == EComponentMobility::Static) && ((HoveredTemplate->Mobility == EComponentMobility::Movable) || (HoveredTemplate->Mobility == EComponentMobility::Stationary)))
+				{
+					// Can't attach Static components to mobile ones
+					Message = LOCTEXT("DropActionToolTip_Error_CannotAttachStatic", "Cannot attach Static components to movable ones.");
+				}
+				else if ((DraggedTemplate->Mobility == EComponentMobility::Stationary) && (HoveredTemplate->Mobility == EComponentMobility::Movable))
+				{
+					// Can't attach Static components to mobile ones
+					Message = LOCTEXT("DropActionToolTip_Error_CannotAttachStationary", "Cannot attach Stationary components to movable ones.");
+				}
+				else if ((NodeData->IsInstancedComponent() && HoveredTemplate->CreationMethod == EComponentCreationMethod::Native && !HoveredTemplate->HasAnyFlags(RF_DefaultSubObject)))
+				{
+					// Can't attach to post-construction C++-added components as they exist outside of the CDO and are not known at SCS execution time
+					Message = LOCTEXT("DropActionToolTip_Error_CannotAttachCPPAdded", "Cannot attach to components added in post-construction C++ code.");
+				}
+				else if (NodeData->IsInstancedComponent() && HoveredTemplate->CreationMethod == EComponentCreationMethod::UserConstructionScript)
+				{
+					// Can't attach to UCS-added components as they exist outside of the CDO and are not known at SCS execution time
+					Message = LOCTEXT("DropActionToolTip_Error_CannotAttachUCSAdded", "Cannot attach to components added in the Construction Script.");
+				}
+				else if (HoveredTemplate->CanAttachAsChild(DraggedTemplate, NAME_None))
+				{
+					// Attach the dragged node(s) to this node
+					if (DraggedNodeData->GetBlueprint() != NodeData->GetBlueprint())
+					{
+						if (DragRowOp->SourceNodes.Num() > 1)
+						{
+							Message = FText::Format(LOCTEXT("DropActionToolTip_AttachToThisNodeFromCopyWithMultipleSelection", "Drop here to copy the selected nodes to new variables and attach them to {0}."), NodeData->GetDisplayName());
+						}
+						else
+						{
+							Message = FText::Format(LOCTEXT("DropActionToolTip_AttachToThisNodeFromCopy", "Drop here to copy {0} to a new variable and attach it to {1}."), DraggedNodeData->GetDisplayName(), NodeData->GetDisplayName());
+						}
+					}
+					else if (DragRowOp->SourceNodes.Num() > 1)
+					{
+						Message = FText::Format(LOCTEXT("DropActionToolTip_AttachToThisNodeWithMultipleSelection", "Drop here to attach the selected components to {0}."), NodeData->GetDisplayName());
+					}
+					else
+					{
+						Message = FText::Format(LOCTEXT("DropActionToolTip_AttachToThisNode", "Drop here to attach {0} to {1}."), DraggedNodeData->GetDisplayName(), NodeData->GetDisplayName());
+					}
+
+					DragRowOp->PendingDropAction = FSubobjectRowDragDropOp::DropAction_AttachTo;
+				}
+				else
+				{
+					// The dropped node cannot be attached to the current node
+					Message = FText::Format(LOCTEXT("DropActionToolTip_Error_TooManyAttachments", "Unable to attach {0} to {1}."), DraggedNodeData->GetDisplayName(), NodeData->GetDisplayName());
+				}
+			}
+		}
+		
+		const FSlateBrush* StatusSymbol = DragRowOp->PendingDropAction != FSubobjectRowDragDropOp::DropAction_None
+			? FEditorStyle::GetBrush(TEXT("Graph.ConnectorFeedback.OK"))
+			: FEditorStyle::GetBrush(TEXT("Graph.ConnectorFeedback.Error"));
+	
+		if (Message.IsEmpty())
+		{
+			DragRowOp->SetFeedbackMessage(nullptr);
+		}
+		else
+		{
+			DragRowOp->SetSimpleFeedbackMessage(StatusSymbol, FLinearColor::White, Message);
+		}
+	}
+	else if (Operation->IsOfType<FExternalDragOperation>() || Operation->IsOfType<FAssetDragDropOp>())
+	{
+		// defer to the tree widget's handler for this type of operation
+		TSharedPtr<SSubobjectEditor> PinnedEditor = SubobjectEditor.Pin();
+		if ( PinnedEditor.IsValid() && PinnedEditor->GetDragDropTree().IsValid() )
+		{
+			// The widget geometry is irrelevant to the tree widget's OnDragEnter
+			PinnedEditor->GetDragDropTree()->OnDragEnter( FGeometry(), DragDropEvent );
+		}
+	}
+}
+
+void SSubobject_RowWidget::HandleOnDragLeave(const FDragDropEvent& DragDropEvent)
+{
+	TSharedPtr<FSubobjectRowDragDropOp> DragRowOp = DragDropEvent.GetOperationAs<FSubobjectRowDragDropOp>();
+	if (DragRowOp.IsValid())
+	{
+		bool bCanReparentAllNodes = true;
+		for(auto SourceNodeIter = DragRowOp->SourceNodes.CreateConstIterator(); SourceNodeIter && bCanReparentAllNodes; ++SourceNodeIter)
+		{
+			FSubobjectEditorTreeNodePtrType DraggedNodePtr = *SourceNodeIter;
+			check(DraggedNodePtr.IsValid());
+			const FSubobjectData* Data = DraggedNodePtr->GetDataSource();
+
+			bCanReparentAllNodes = Data->CanReparent();
+		}
+
+		// Only clear the tooltip text if all dragged nodes support it
+		if(bCanReparentAllNodes)
+		{
+			TSharedPtr<SWidget> NoWidget;
+			DragRowOp->SetFeedbackMessage(NoWidget);
+			DragRowOp->PendingDropAction = FSubobjectRowDragDropOp::DropAction_None;
+		}
+	}
+}
+
+FReply SSubobject_RowWidget::HandleOnDragDetected(const FGeometry& MyGeometry, const FPointerEvent& MouseEvent)
+{
+	TSharedPtr<SSubobjectEditor> SubobjectEditorPtr = SubobjectEditor.Pin();
+	if (MouseEvent.IsMouseButtonDown(EKeys::LeftMouseButton)
+        && SubobjectEditorPtr.IsValid()
+        && SubobjectEditorPtr->IsEditingAllowed()) //can only drag when editing
+    {
+		TArray<FSubobjectEditorTreeNodePtrType> SelectedNodePtrs = SubobjectEditorPtr->GetSelectedNodes();
+		if (SelectedNodePtrs.Num() == 0)
+		{
+			SelectedNodePtrs.Add(GetSubobjectPtr());
+		}
+
+		FSubobjectEditorTreeNodePtrType FirstNode = SelectedNodePtrs[0];
+		const FSubobjectData* Data = FirstNode->GetDataSource();
+
+		if (Data && FirstNode->IsComponentNode())
+		{
+			// Do not use the Blueprint from FirstNode, it may still be referencing the parent.
+			UBlueprint* Blueprint = SubobjectEditorPtr->GetBlueprint();
+			const FName VariableName = Data->GetVariableName();
+			UStruct* VariableScope = (Blueprint != nullptr) ? Blueprint->SkeletonGeneratedClass : nullptr;
+
+			TSharedRef<FSubobjectRowDragDropOp> Operation = FSubobjectRowDragDropOp::New(VariableName, VariableScope, FNodeCreationAnalytic());
+			Operation->SetCtrlDrag(true); // Always put a getter
+			Operation->PendingDropAction = FSubobjectRowDragDropOp::DropAction_None;
+			Operation->SourceNodes = SelectedNodePtrs;
+
+			return FReply::Handled().BeginDragDrop(Operation);
+		}
+    }
+	
+	return FReply::Unhandled();
+}
+
+TOptional<EItemDropZone> SSubobject_RowWidget::HandleOnCanAcceptDrop(const FDragDropEvent& DragDropEvent, EItemDropZone DropZone, FSubobjectEditorTreeNodePtrType TargetItem)
+{
+	TOptional<EItemDropZone> ReturnDropZone;
+
+	TSharedPtr<FDragDropOperation> Operation = DragDropEvent.GetOperation();
+	if (Operation.IsValid())
+	{
+		if (Operation->IsOfType<FSubobjectRowDragDropOp>() && ( Cast<USceneComponent>(GetSubobjectPtr()->GetComponentTemplate()) != nullptr ))
+		{
+			TSharedPtr<FSubobjectRowDragDropOp> DragRowOp = StaticCastSharedPtr<FSubobjectRowDragDropOp>(Operation);
+			check(DragRowOp.IsValid());
+
+			if (DragRowOp->PendingDropAction != FSubobjectRowDragDropOp::DropAction_None)
+			{
+				ReturnDropZone = EItemDropZone::OntoItem;
+			}
+		}
+		else if (Operation->IsOfType<FExternalDragOperation>() || Operation->IsOfType<FAssetDragDropOp>())
+		{
+			ReturnDropZone = EItemDropZone::OntoItem;
+		}
+	}
+
+	return ReturnDropZone;
+}
+
+FReply SSubobject_RowWidget::HandleOnAcceptDrop(const FDragDropEvent& DragDropEvent, EItemDropZone DropZone, FSubobjectEditorTreeNodePtrType TargetItem)
+{
+	TSharedPtr<FDragDropOperation> Operation = DragDropEvent.GetOperation();
+	if (!Operation.IsValid())
+	{
+		return FReply::Handled();
+	}
+
+	TSharedPtr<SSubobjectEditor> SubobjectEditorPtr = SubobjectEditor.Pin();
+	if(!SubobjectEditorPtr.IsValid())
+	{
+		return FReply::Handled();
+	}
+	
+	if (Operation->IsOfType<FSubobjectRowDragDropOp>() && (Cast<USceneComponent>(GetSubobjectPtr()->GetComponentTemplate()) != nullptr))
+	{
+		TSharedPtr<FSubobjectRowDragDropOp> DragRowOp = StaticCastSharedPtr<FSubobjectRowDragDropOp>(Operation);	
+		check(DragRowOp.IsValid());
+
+		switch(DragRowOp->PendingDropAction)
+		{
+		case FSubobjectRowDragDropOp::DropAction_AttachTo:
+			SubobjectEditorPtr->OnAttachToDropAction(GetSubobjectPtr(), DragRowOp->SourceNodes);
+			break;
+			
+		case FSubobjectRowDragDropOp::DropAction_DetachFrom:
+			SubobjectEditorPtr->OnDetachFromDropAction(DragRowOp->SourceNodes);
+			break;
+
+		case FSubobjectRowDragDropOp::DropAction_MakeNewRoot:
+			check(DragRowOp->SourceNodes.Num() == 1);
+			SubobjectEditorPtr->OnMakeNewRootDropAction(DragRowOp->SourceNodes[0]);
+			break;
+
+		case FSubobjectRowDragDropOp::DropAction_AttachToOrMakeNewRoot:
+			{
+				check(DragRowOp->SourceNodes.Num() == 1);
+				FWidgetPath WidgetPath = DragDropEvent.GetEventPath() != nullptr ? *DragDropEvent.GetEventPath() : FWidgetPath();
+				FSlateApplication::Get().PushMenu(
+                    SharedThis(this),
+                    WidgetPath,
+                    SubobjectEditorPtr->BuildSceneRootDropActionMenu(GetSubobjectPtr(), DragRowOp->SourceNodes[0]).ToSharedRef(),
+                    FSlateApplication::Get().GetCursorPos(),
+                    FPopupTransitionEffect(FPopupTransitionEffect::TypeInPopup)
+                );
+			}
+			break;
+
+		case FSubobjectRowDragDropOp::DropAction_None:
+		default:
+            break;
+		}
+	}
+	else if (Operation->IsOfType<FExternalDragOperation>() || Operation->IsOfType<FAssetDragDropOp>())
+	{
+		// defer to the tree widget's handler for this type of operation
+		TSharedPtr<SSubobjectEditor> PinnedEditor = SubobjectEditor.Pin();
+		TSharedPtr<SSubobjectEditorDragDropTree> TreeWidget = PinnedEditor.IsValid() ? PinnedEditor->GetDragDropTree() : nullptr;
+		if (TreeWidget.IsValid())
+		{
+			// The widget geometry is irrelevant to the tree widget's OnDrop
+			PinnedEditor->GetDragDropTree()->OnDrop(FGeometry(), DragDropEvent);
+		}
+	}
+	
+	return FReply::Handled();
+}
+
+BEGIN_SLATE_FUNCTION_BUILD_OPTIMIZATION
+void SSubobject_RowWidget::AddToToolTipInfoBox(const TSharedRef<SVerticalBox>& InfoBox, const FText& Key, TSharedRef<SWidget> ValueIcon, const TAttribute<FText>& Value, bool bImportant)
+{
+	InfoBox->AddSlot()
+	       .AutoHeight()
+	       .Padding(0, 1)
+	[
+		SNew(SHorizontalBox)
+
+		+ SHorizontalBox::Slot()
+		  .AutoWidth()
+		  .Padding(0, 0, 4, 0)
+		[
+			SNew(STextBlock)
+			.TextStyle(FEditorStyle::Get(),
+			           bImportant ? "SCSEditor.ComponentTooltip.ImportantLabel" : "SCSEditor.ComponentTooltip.Label")
+		.Text(FText::Format(LOCTEXT("AssetViewTooltipFormat", "{0}:"), Key))
+		]
+
+		+ SHorizontalBox::Slot()
+		.AutoWidth()
+		[
+			ValueIcon
+		]
+
+		+ SHorizontalBox::Slot()
+		.AutoWidth()
+		[
+			SNew(STextBlock)
+			.TextStyle(FEditorStyle::Get(),
+			           bImportant ? "SCSEditor.ComponentTooltip.ImportantValue" : "SCSEditor.ComponentTooltip.Value")
+		.Text(Value)
+		]
+	];
+}
+END_SLATE_FUNCTION_BUILD_OPTIMIZATION
+
+TSharedRef<SToolTip> SSubobject_RowWidget::CreateToolTipWidget() const
+{
+	FSubobjectEditorTreeNodePtrType TreeNode = GetSubobjectPtr();
+	check(TreeNode);
+	const FSubobjectData* TreeNodeData = TreeNode->GetDataSource();
+	check(TreeNodeData);
+	
+	if (TreeNode->IsComponentNode())
+	{
+		return CreateComponentTooltipWidget(TreeNode);
+	}
+	else
+	{
+		return CreateActorTooltipWidget(TreeNode);
+	}
+}
+
+TSharedRef<SToolTip> SSubobject_RowWidget::CreateComponentTooltipWidget(const FSubobjectEditorTreeNodePtrType& InNode) const
+{
+	// Create a box to hold every line of info in the body of the tooltip
+	TSharedRef<SVerticalBox> InfoBox = SNew(SVerticalBox);
+
+	if (FSubobjectEditorTreeNodePtrType TreeNode = GetSubobjectPtr())
+	{
+		const FSubobjectData* TreeNodeData = TreeNode->GetDataSource();
+		check(TreeNodeData);
+		if (TreeNode->IsComponentNode())
+		{
+			if (const UActorComponent* Template = TreeNode->GetComponentTemplate())
+			{		
+				UClass* TemplateClass = Template->GetClass();
+				FText ClassTooltip = TemplateClass->GetToolTipText(/*bShortTooltip=*/ true);
+
+				InfoBox->AddSlot()
+				       .AutoHeight()
+				       .HAlign(HAlign_Center)
+				       .Padding(FMargin(0, 2, 0, 4))
+				[
+					SNew(STextBlock)
+						.TextStyle(FEditorStyle::Get(), "SCSEditor.ComponentTooltip.ClassDescription")
+						.Text(ClassTooltip)
+						.WrapTextAt(400.0f)
+				];
+			}
+
+			// Add introduction point
+			AddToToolTipInfoBox(InfoBox, LOCTEXT("TooltipAddType", "Source"), SNullWidget::NullWidget, TAttribute<FText>::Create(TAttribute<FText>::FGetter::CreateSP(this, &SSubobject_RowWidget::GetComponentAddSourceToolTipText)), false);
+			if (TreeNodeData->IsInheritedComponent())
+			{
+				AddToToolTipInfoBox(InfoBox, LOCTEXT("TooltipIntroducedIn", "Introduced in"), SNullWidget::NullWidget, TAttribute<FText>::Create(TAttribute<FText>::FGetter::CreateSP(this, &SSubobject_RowWidget::GetIntroducedInToolTipText)), false);
+			}
+
+			// Add Underlying Component Name for Native Components
+			if (TreeNodeData->IsNativeComponent())
+			{
+				AddToToolTipInfoBox(InfoBox, LOCTEXT("TooltipNativeComponentName", "Native Component Name"), SNullWidget::NullWidget, TAttribute<FText>::Create(TAttribute<FText>::FGetter::CreateSP(this, &SSubobject_RowWidget::GetNativeComponentNameToolTipText)), false);
+			}
+
+			// Add mobility
+			TSharedRef<SImage> MobilityIcon = SNew(SImage).Image(this, &SSubobject_RowWidget::GetMobilityIconImage);
+			AddToToolTipInfoBox(InfoBox, LOCTEXT("TooltipMobility", "Mobility"), MobilityIcon, TAttribute<FText>::Create(TAttribute<FText>::FGetter::CreateSP(this, &SSubobject_RowWidget::GetMobilityToolTipText)), false);
+
+			// Add asset if applicable to this node
+			if (GetAssetVisibility() == EVisibility::Visible)
+			{
+				InfoBox->AddSlot()[SNew(SSpacer).Size(FVector2D(1.0f, 8.0f))];
+				AddToToolTipInfoBox(InfoBox, LOCTEXT("TooltipAsset", "Asset"), SNullWidget::NullWidget,
+				                    TAttribute<FText>(this, &SSubobject_RowWidget::GetAssetName), false);
+			}
+
+			// If the component is marked as editor only, then display that info here
+			AddToToolTipInfoBox(InfoBox, LOCTEXT("TooltipEditorOnly", "Editor Only"), SNullWidget::NullWidget, TAttribute<FText>::Create(TAttribute<FText>::FGetter::CreateSP(this, &SSubobject_RowWidget::GetComponentEditorOnlyTooltipText)), false);
+		}
+	}
+
+	TSharedRef<SBorder> TooltipContent = SNew(SBorder)
+		.BorderImage(FEditorStyle::GetBrush("NoBorder"))
+		.Padding(0)
+	[
+		SNew(SVerticalBox)
+
+		+ SVerticalBox::Slot()
+		  .AutoHeight()
+		  .Padding(0, 0, 0, 4)
+		[
+			SNew(SVerticalBox)
+			+ SVerticalBox::Slot()
+			.AutoHeight()
+			[
+				SNew(SHorizontalBox)
+				+ SHorizontalBox::Slot()
+				  .AutoWidth()
+				  .VAlign(VAlign_Center)
+				  .Padding(2)
+				[
+					SNew(STextBlock)
+					.TextStyle(FEditorStyle::Get(), "SCSEditor.ComponentTooltip.Title")
+					.Text(this, &SSubobject_RowWidget::GetTooltipText)
+				]
+			]
+		]
+
+		+ SVerticalBox::Slot()
+		.AutoHeight()
+		[
+			SNew(SBorder)
+			.BorderImage(FEditorStyle::GetBrush("NoBorder"))
+			.Padding(2)
+			[
+				InfoBox
+			]
+		]
+	];
+
+	return IDocumentation::Get()->CreateToolTip(TAttribute<FText>(this, &SSubobject_RowWidget::GetTooltipText),
+	                                            TooltipContent, InfoBox, GetDocumentationLink(),
+	                                            GetDocumentationExcerptName());
+}
+
+TSharedRef<SToolTip> SSubobject_RowWidget::CreateActorTooltipWidget(const FSubobjectEditorTreeNodePtrType& InNode) const
+{
+	// Create a box to hold every line of info in the body of the tooltip
+	TSharedRef<SVerticalBox> InfoBox = SNew(SVerticalBox);
+
+	// Add class
+	AddToToolTipInfoBox(InfoBox, LOCTEXT("TooltipClass", "Class"), SNullWidget::NullWidget, TAttribute<FText>::Create(TAttribute<FText>::FGetter::CreateSP(this, &SSubobject_RowWidget::GetActorClassNameText)), false);
+
+	// Add super class
+	AddToToolTipInfoBox(InfoBox, LOCTEXT("TooltipSuperClass", "Parent Class"), SNullWidget::NullWidget, TAttribute<FText>::Create(TAttribute<FText>::FGetter::CreateSP(this, &SSubobject_RowWidget::GetActorSuperClassNameText)), false);
+
+	// Add mobility
+	AddToToolTipInfoBox(InfoBox, LOCTEXT("TooltipMobility", "Mobility"), SNullWidget::NullWidget, TAttribute<FText>::Create(TAttribute<FText>::FGetter::CreateSP(this, &SSubobject_RowWidget::GetActorMobilityText)), false);
+
+	TSharedRef<SBorder> TooltipContent = SNew(SBorder)
+        .BorderImage(FEditorStyle::GetBrush("NoBorder"))
+        .Padding(0)
+        [
+            SNew(SVerticalBox)
+
+            + SVerticalBox::Slot()
+            .AutoHeight()
+            .Padding(0, 0, 0, 4)
+            [
+                SNew(SVerticalBox)
+
+                + SVerticalBox::Slot()
+                .AutoHeight()
+                [
+                    SNew(SHorizontalBox)
+
+                    + SHorizontalBox::Slot()
+                    .AutoWidth()
+                    .VAlign(VAlign_Center)
+                    .Padding(4)
+                    [
+                        SNew(STextBlock)
+                        .TextStyle(FEditorStyle::Get(), "SCSEditor.ComponentTooltip.Title")
+                        .Text(this, &SSubobject_RowWidget::GetActorDisplayText)
+                    ]
+                ]
+            ]
+
+            + SVerticalBox::Slot()
+            .AutoHeight()
+            [
+                SNew(SBorder)
+                .BorderImage(FEditorStyle::GetBrush("NoBorder"))
+                .Padding(4)
+                [
+                    InfoBox
+                ]
+            ]
+        ];
+
+	return IDocumentation::Get()->CreateToolTip(TAttribute<FText>(this, &SSubobject_RowWidget::GetActorDisplayText), TooltipContent, InfoBox, TEXT(""), TEXT(""));
+}
+
+FText SSubobject_RowWidget::GetTooltipText() const
+{
+	const FSubobjectData* Data = GetSubobjectPtr()->GetDataSource();
+
+	if(!Data)
+	{
+		return FText::GetEmpty();
+	}
+	
+	if (Data->IsDefaultSceneRoot())
+	{
+		if (Data->IsInheritedComponent())
+		{
+			return LOCTEXT("InheritedDefaultSceneRootToolTip",
+			               "This is the default scene root component. It cannot be copied, renamed or deleted.\nIt has been inherited from the parent class, so its properties cannot be edited here.\nNew scene components will automatically be attached to it.");
+		}
+		else
+		{
+			return LOCTEXT("DefaultSceneRootToolTip",
+			               "This is the default scene root component. It cannot be copied, renamed or deleted.\nIt can be replaced by drag/dropping another scene component over it.");
+		}
+	}
+	else
+	{
+		const UClass* Class = (Data->GetComponentTemplate() != nullptr)
+			                      ? Data->GetComponentTemplate()->GetClass()
+			                      : nullptr;
+		const FText ClassDisplayName = FBlueprintEditorUtils::GetFriendlyClassDisplayName(Class);
+
+		FFormatNamedArguments Args;
+		Args.Add(TEXT("ClassName"), ClassDisplayName);
+		Args.Add(TEXT("NodeName"), FText::FromString(Data->GetDisplayString()));
+
+		return FText::Format(LOCTEXT("ComponentTooltip", "{NodeName} ({ClassName})"), Args);
+	}
+}
+
+FText SSubobject_RowWidget::GetActorClassNameText() const
+{
+	FSubobjectEditorTreeNodePtrType NodePtr = GetSubobjectPtr();
+	const FSubobjectData* Data = NodePtr ? NodePtr->GetDataSource() : nullptr;
+	if(Data)
+	{
+		if (const AActor* DefaultActor = Data->GetObject<AActor>())
+		{
+			return FText::FromString(DefaultActor->GetClass()->GetName());
+		}
+	}
+	
+	return FText::GetEmpty();
+}
+
+FText SSubobject_RowWidget::GetActorSuperClassNameText() const
+{
+	FSubobjectEditorTreeNodePtrType NodePtr = GetSubobjectPtr();
+	const FSubobjectData* Data = NodePtr ? NodePtr->GetDataSource() : nullptr;
+	if (Data)
+	{
+		if (const AActor* DefaultActor = Data->GetObject<AActor>())
+		{
+			return FText::FromString(DefaultActor->GetClass()->GetSuperClass()->GetName());
+		}
+	}
+
+	return FText::GetEmpty();
+}
+
+FText SSubobject_RowWidget::GetActorMobilityText() const
+{
+	FSubobjectEditorTreeNodePtrType NodePtr = GetSubobjectPtr();
+	const FSubobjectData* Data = NodePtr ? NodePtr->GetDataSource() : nullptr;
+	if (Data)
+	{
+		if (const AActor* DefaultActor = Data->GetObject<AActor>())
+		{
+			USceneComponent* RootComponent = DefaultActor->GetRootComponent();
+
+			if (RootComponent != nullptr)
+			{
+				if (RootComponent->Mobility == EComponentMobility::Static)
+				{
+					return LOCTEXT("ComponentMobility_Static", "Static");
+				}
+				else if (RootComponent->Mobility == EComponentMobility::Stationary)
+				{
+					return LOCTEXT("ComponentMobility_Stationary", "Stationary");
+				}
+				else if (RootComponent->Mobility == EComponentMobility::Movable)
+				{
+					return LOCTEXT("ComponentMobility_Movable", "Movable");
+				}
+			}
+			else
+			{
+				return LOCTEXT("ComponentMobility_NoRoot", "No root component, unknown mobility");
+			}
+		}
+	}
+
+	return FText::GetEmpty();
+}
+
+FText SSubobject_RowWidget::GetObjectContextText() const
+{
+	FSubobjectEditorTreeNodePtrType NodePtr = GetSubobjectPtr();
+    const FSubobjectData* Data = NodePtr ? NodePtr->GetDataSource() : nullptr;
+    if (Data)
+    {
+	   return Data->GetDisplayNameContextModifiers(GetDefault<UEditorStyleSettings>()->bShowNativeComponentNames);
+    }
+	return FText::GetEmpty();
+}
+
+FString SSubobject_RowWidget::GetDocumentationLink() const
+{
+	FSubobjectEditorTreeNodePtrType NodePtr = GetSubobjectPtr();
+	const FSubobjectData* Data = NodePtr ? NodePtr->GetDataSource() : nullptr;
+	
+	if (Data && (NodePtr == SubobjectEditor.Pin()->GetSceneRootNode() || Data->IsInheritedComponent()))
+	{
+		return TEXT("Shared/Editors/BlueprintEditor/ComponentsMode");
+	}
+
+	return TEXT("");
+}
+
+FString SSubobject_RowWidget::GetDocumentationExcerptName() const
+{
+	FSubobjectEditorTreeNodePtrType NodePtr = GetSubobjectPtr();
+	const FSubobjectData* Data = NodePtr ? NodePtr->GetDataSource() : nullptr;
+
+	if(NodePtr && Data)
+	{
+		if (NodePtr == SubobjectEditor.Pin()->GetSceneRootNode())
+		{
+			return TEXT("RootComponent");
+		}
+		else if (Data->IsNativeComponent())
+		{
+			return TEXT("NativeComponents");
+		}
+		else if (Data->IsInheritedComponent())
+		{
+			return TEXT("InheritedComponents");
+		}	
+	}
+
+	return TEXT("");
+}
+
+FText SSubobject_RowWidget::GetMobilityToolTipText() const
+{
+	const FSubobjectEditorTreeNodePtrType NodePtr = GetSubobjectPtr();
+	const FSubobjectData* Data = NodePtr ? NodePtr->GetDataSource() : nullptr;
+	
+	if (Data)
+	{
+		return Data->GetMobilityToolTipText();
+	}
+
+	return LOCTEXT("ErrorNoMobilityTooltip", "Invalid component");
+}
+
+FSlateBrush const* SSubobject_RowWidget::GetMobilityIconImage() const
+{
+	if (FSubobjectEditorTreeNodePtrType NodePtr = GetSubobjectPtr())
+	{
+		if(FSubobjectData* Data = NodePtr->GetDataSource())
+		{
+			if (const USceneComponent* SceneComponentTemplate = Cast<USceneComponent>(Data->GetComponentTemplate()))
+			{
+				if (SceneComponentTemplate->Mobility == EComponentMobility::Movable)
+				{
+					return FEditorStyle::GetBrush(TEXT("ClassIcon.MovableMobilityIcon"));
+				}
+				else if (SceneComponentTemplate->Mobility == EComponentMobility::Stationary)
+				{
+					return FEditorStyle::GetBrush(TEXT("ClassIcon.StationaryMobilityIcon"));
+				}
+				// static components don't get an icon (because static is the most common
+				// mobility type, and we'd like to keep the icon clutter to a minimum)
+			}
+		}	
+	}
+
+	return nullptr;
+}
+
+FText SSubobject_RowWidget::GetIntroducedInToolTipText() const
+{
+	FText IntroducedInTooltip = LOCTEXT("IntroducedInThisBPTooltip", "this class");
+	if (const FSubobjectData* TreeNode = SubobjectPtr->GetDataSource())
+	{
+		return TreeNode->GetIntroducedInToolTipText();
+	}
+
+	return IntroducedInTooltip;
+}
+
+FText SSubobject_RowWidget::GetComponentAddSourceToolTipText() const
+{
+	FText NodeType;
+	
+	if (const FSubobjectData* TreeNode = SubobjectPtr->GetDataSource())
+	{
+		if (TreeNode->IsInheritedComponent())
+		{
+			if (TreeNode->IsNativeComponent())
+			{
+				NodeType = LOCTEXT("InheritedNativeComponent", "Inherited (C++)");
+			}
+			else
+			{
+				NodeType = LOCTEXT("InheritedBlueprintComponent", "Inherited (Blueprint)");
+			}
+		}
+		else
+		{
+			if (TreeNode->IsInstancedComponent())
+			{
+				NodeType = LOCTEXT("ThisInstanceAddedComponent", "This actor instance");
+			}
+			else
+			{
+				NodeType = LOCTEXT("ThisBlueprintAddedComponent", "This Blueprint");
+			}
+		}
+	}
+
+	return NodeType;
+}
+
+FText SSubobject_RowWidget::GetComponentEditorOnlyTooltipText() const
+{
+	const FSubobjectData* TreeNode = SubobjectPtr->GetDataSource();
+	return TreeNode ? TreeNode->GetComponentEditorOnlyTooltipText() : LOCTEXT("ComponentEditorOnlyFalse", "False");
+}
+
+FText SSubobject_RowWidget::GetNativeComponentNameToolTipText() const
+{
+	const FSubobjectData* TreeNode = SubobjectPtr->GetDataSource();
+	const UActorComponent* Template = TreeNode ? TreeNode->GetComponentTemplate() : nullptr;
+
+	if (Template)
+	{
+		return FText::FromName(Template->GetFName());
+	}
+	else
+	{
+		return FText::GetEmpty();
+	}
+}
+
+FText SSubobject_RowWidget::GetActorDisplayText() const
+{
+	const FSubobjectData* TreeNode = SubobjectPtr->GetDataSource();
+
+	if (TreeNode)
+	{
+		return TreeNode->GetActorDisplayText();
+	}
+	
+	return FText::GetEmpty();
+}
+
+void SSubobject_RowWidget::OnNameTextCommit(const FText& InNewName, ETextCommit::Type InTextCommit)
+{
+	// Ask the subsystem to rename this
+	FSubobjectEditorTreeNodePtrType NodePtr = GetSubobjectPtr();
+	if (ensure(NodePtr))
+	{
+		// If there is already an ongoing transaction then use that, otherwise make a
+		// new transaction for renaming the component
+		TUniquePtr<FScopedTransaction> Transaction;
+		NodePtr->GetOngoingCreateTransaction(Transaction);
+
+		// If a 'create' transaction is opened, the rename will be folded into it and will be invisible to the
+		// 'undo' as create + give a name is really just one operation from the user point of view.
+		FScopedTransaction TransactionContext(LOCTEXT("RenameComponentVariable", "Rename Component Variable"));
+
+		// Ask the subsystem to verify this rename option for us
+		USubobjectDataSubsystem* System = USubobjectDataSubsystem::Get();
+		check(System);
+		const bool bSuccess = System->RenameSubobject(NodePtr->GetDataHandle(), InNewName);
+		
+		Transaction.Reset();
+	}
+}
+
+bool SSubobject_RowWidget::OnNameTextVerifyChanged(const FText& InNewText, FText& OutErrorMessage)
+{
+	const FSubobjectEditorTreeNodePtrType NodePtr = GetSubobjectPtr();
+	if (ensure(NodePtr))
+	{
+		// Ask the subsystem to verify this rename option for us
+		const USubobjectDataSubsystem* System = USubobjectDataSubsystem::Get();
+		check(System);
+		return System->IsValidRename(NodePtr->GetDataHandle(), InNewText, OutErrorMessage);
+	}
+	
+	return false;
+}
+
+bool SSubobject_RowWidget::IsReadOnly() const
+{
+	FSubobjectEditorTreeNodePtrType NodePtr = GetSubobjectPtr();
+	const FSubobjectData* Data = NodePtr ? NodePtr->GetDataSource() : nullptr;
+	if (Data)
+	{
+		return !Data->CanRename() || (SubobjectEditor.IsValid() && !SubobjectEditor.Pin()->IsEditingAllowed());
+	}
+
+	// If there is no valid data then default to read only
+	return true;
+}
+
+END_SLATE_FUNCTION_BUILD_OPTIMIZATION
+
+////////////////////////////////////////////////
+// SSubobjectEditor
+
+FMenuBuilder SSubobjectEditor::CreateMenuBuilder() const
+{
+	// Menu builder for editing a blueprint
+	FMenuBuilder EditBlueprintMenuBuilder(true, nullptr);
+
+	EditBlueprintMenuBuilder.BeginSection(
+		NAME_None, LOCTEXT("EditBlueprintMenu_ExistingBlueprintHeader", "Existing Blueprint"));
+
+	EditBlueprintMenuBuilder.AddMenuEntry
+	(
+		LOCTEXT("OpenBlueprintEditor", "Open Blueprint Editor"),
+		LOCTEXT("OpenBlueprintEditor_ToolTip", "Opens the blueprint editor for this asset"),
+		FSlateIcon(),
+		FUIAction(
+			FExecuteAction::CreateSP(this, &SSubobjectEditor::OnOpenBlueprintEditor, /*bForceCodeEditing=*/ false))
+	);
+
+	EditBlueprintMenuBuilder.AddMenuEntry
+	(
+		LOCTEXT("OpenBlueprintEditorScriptMode", "Add or Edit Script"),
+		LOCTEXT("OpenBlueprintEditorScriptMode_ToolTip",
+		        "Opens the blueprint editor for this asset, showing the event graph"),
+		FSlateIcon(),
+		FUIAction(FExecuteAction::CreateSP(this, &SSubobjectEditor::OnOpenBlueprintEditor, /*bForceCodeEditing=*/ true))
+	);
+
+	return EditBlueprintMenuBuilder;
+}
+
+void SSubobjectEditor::ConstructTreeWidget()
+{
+	TSharedPtr<SHeaderRow> HeaderRow = SNew(SHeaderRow)
+		+ SHeaderRow::Column(SubobjectTree_ColumnName_ComponentClass)
+		  .DefaultLabel(LOCTEXT("Class", "Class"))
+		  .FillWidth(4);
+
+	TreeWidget = SNew(SSubobjectEditorDragDropTree)
+		.SubobjectEditor(this)
+		.ToolTipText(LOCTEXT("DropAssetToAddComponent", "Drop asset here to add a component."))
+		.TreeItemsSource(&RootNodes)
+		.SelectionMode(ESelectionMode::Multi)
+		.OnGenerateRow(this, &SSubobjectEditor::MakeTableRowWidget)
+		.OnGetChildren(this, &SSubobjectEditor::OnGetChildrenForTree)
+		.OnSetExpansionRecursive(this, &SSubobjectEditor::SetItemExpansionRecursive)
+		.OnSelectionChanged(this, &SSubobjectEditor::OnTreeSelectionChanged)
+		.OnContextMenuOpening(this, &SSubobjectEditor::CreateContextMenu)
+		.OnItemScrolledIntoView(this, &SSubobjectEditor::OnItemScrolledIntoView)
+		.OnMouseButtonDoubleClick(this, &SSubobjectEditor::HandleItemDoubleClicked)
+		.ClearSelectionOnClick(ClearSelectionOnClick())
+		.OnTableViewBadState(this, &SSubobjectEditor::DumpTree)
+		.ItemHeight(24)
+		.HeaderRow
+	    (
+			HeaderRow
+	    );
+
+	HeaderRow->SetVisibility(EVisibility::Collapsed);
+}
+
+void SSubobjectEditor::CreateCommandList()
+{
+	CommandList = MakeShareable(new FUICommandList);
+
+	CommandList->MapAction(FGenericCommands::Get().Cut,
+       FUIAction(FExecuteAction::CreateSP(this, &SSubobjectEditor::CutSelectedNodes),
+       FCanExecuteAction::CreateSP(this, &SSubobjectEditor::CanCutNodes))
+    );
+
+	CommandList->MapAction(FGenericCommands::Get().Copy,
+		FUIAction(FExecuteAction::CreateSP(this, &SSubobjectEditor::CopySelectedNodes),
+		FCanExecuteAction::CreateSP(this, &SSubobjectEditor::CanCopyNodes))
+	);
+
+	CommandList->MapAction(FGenericCommands::Get().Paste,
+	    FUIAction(FExecuteAction::CreateSP(this, &SSubobjectEditor::PasteNodes),
+	    FCanExecuteAction::CreateSP(this, &SSubobjectEditor::CanPasteNodes))
+	);
+
+	CommandList->MapAction(FGenericCommands::Get().Duplicate,
+	    FUIAction(FExecuteAction::CreateSP(this, &SSubobjectEditor::OnDuplicateComponent),
+	    FCanExecuteAction::CreateSP(this, &SSubobjectEditor::CanDuplicateComponent))
+	);
+
+	CommandList->MapAction(FGenericCommands::Get().Delete,
+	    FUIAction(FExecuteAction::CreateSP(this, &SSubobjectEditor::OnDeleteNodes),
+		FCanExecuteAction::CreateSP(this, &SSubobjectEditor::CanDeleteNodes))
+	);
+
+	CommandList->MapAction(FGenericCommands::Get().Rename,
+	    FUIAction(FExecuteAction::CreateSP(this, &SSubobjectEditor::OnRenameComponent),
+	    FCanExecuteAction::CreateSP(this, &SSubobjectEditor::CanRenameComponent))
+	);
+
+	CommandList->MapAction(FGraphEditorCommands::Get().GetFindReferences(),
+		FUIAction(FExecuteAction::CreateSP(this, &SSubobjectEditor::OnFindReferences))
+	);
+}
+
+bool SSubobjectEditor::IsComponentSelected(const UPrimitiveComponent* PrimComponent) const
+{
+	check(PrimComponent);
+
+	if (TreeWidget.IsValid())
+	{
+		FSubobjectEditorTreeNodePtrType NodePtr = FindSlateNodeForObject(PrimComponent, false);
+		if (NodePtr.IsValid())
+		{
+			return TreeWidget->IsItemSelected(NodePtr);
+		}
+		else
+		{
+			UChildActorComponent* PossiblySelectedComponent = nullptr;
+			AActor* ComponentOwner = PrimComponent->GetOwner();
+			while (ComponentOwner->IsChildActor())
+			{
+				PossiblySelectedComponent = ComponentOwner->GetParentComponent();
+				ComponentOwner = ComponentOwner->GetParentActor();
+			}
+
+			if (PossiblySelectedComponent)
+			{
+				NodePtr = FindSlateNodeForObject(PossiblySelectedComponent, false);
+				if (NodePtr.IsValid())
+				{
+					return TreeWidget->IsItemSelected(NodePtr);
+				}
+			}
+		}
+	}
+	
+	return false;
+}
+
+void SSubobjectEditor::SetSelectionOverride(UPrimitiveComponent* PrimComponent) const
+{
+	PrimComponent->SelectionOverrideDelegate = UPrimitiveComponent::FSelectionOverride::CreateSP(this, &SSubobjectEditor::IsComponentSelected);
+	PrimComponent->PushSelectionToProxy();
+}
+
+TSharedRef<ITableRow> SSubobjectEditor::MakeTableRowWidget(FSubobjectEditorTreeNodePtrType InNodePtr, const TSharedRef<STableViewBase>& OwnerTable)
+{
+	FGraphNodeMetaData TagMeta(TEXT("TableRow"));
+	const UActorComponent* Template = InNodePtr && !InNodePtr->IsSeperator()
+		                                  ? InNodePtr->GetDataSource()->GetComponentTemplate()
+		                                  : nullptr;
+	if (Template)
+	{
+		TagMeta.FriendlyName = FString::Printf(TEXT("TableRow,%s,0"), *Template->GetReadableName());
+	}
+
+	return SNew(SSubobject_RowWidget, SharedThis(this), InNodePtr, OwnerTable)
+		.AddMetaData<FTutorialMetaData>(TagMeta);
+}
+
+EVisibility SSubobjectEditor::GetComponentsTreeVisibility() const
+{
+	return (UICustomization.IsValid() && UICustomization->HideComponentsTree())
+		       ? EVisibility::Collapsed
+		       : EVisibility::Visible;
+}
+
+void SSubobjectEditor::OnGetChildrenForTree(FSubobjectEditorTreeNodePtrType InNodePtr, TArray<FSubobjectEditorTreeNodePtrType>& OutChildren)
+{
+	if (InNodePtr)
+	{
+		const TArray<FSubobjectEditorTreeNodePtrType>& Children = InNodePtr->GetChildren();
+		OutChildren.Reserve(Children.Num());
+
+		OutChildren = Children;
+	}
+	else
+	{
+		OutChildren.Empty();
+	}
+}
+
+TArray<FSubobjectEditorTreeNodePtrType> SSubobjectEditor::GetSelectedNodes() const
+{
+	TArray<FSubobjectEditorTreeNodePtrType> SelectedTreeNodes = TreeWidget->GetSelectedItems();
+
+	struct FCompareSelectedSubobjectEditorTreeNodes
+	{
+		FORCEINLINE bool operator()(const FSubobjectEditorTreeNodePtrType& A,
+		                            const FSubobjectEditorTreeNodePtrType& B) const
+		{
+			return B.IsValid() && B->IsAttachedTo(A);
+		}
+	};
+
+	// Ensure that nodes are ordered from parent to child (otherwise they are sorted in the order that they were selected)
+	SelectedTreeNodes.Sort(FCompareSelectedSubobjectEditorTreeNodes());
+
+	return SelectedTreeNodes;
+}
+
+TArray<FSubobjectDataHandle> SSubobjectEditor::GetSelectedHandles() const
+{
+	TArray<FSubobjectEditorTreeNodePtrType> SelectedNodes = GetSelectedNodes();
+	TArray<FSubobjectDataHandle> OutHandles;
+	for(FSubobjectEditorTreeNodePtrType& Ptr : SelectedNodes)
+	{
+		OutHandles.Add(Ptr->GetDataHandle());
+	}
+	return OutHandles;
+}
+
+FSubobjectEditorTreeNodePtrType SSubobjectEditor::GetSceneRootNode() const
+{
+	// Get the first scene component in the root nodes
+	if(const AActor* ActorContext = Cast<AActor>(GetObjectContext()))
+	{
+		if(const USceneComponent* RootComp = ActorContext->GetRootComponent())
+		{
+			TSet<FSubobjectEditorTreeNodePtrType> VisitedNodes;
+			FSubobjectDataHandle RootHandle;
+			
+			DepthFirstTraversal(RootNodes[0], VisitedNodes,
+                [&RootComp, &RootHandle](
+                const FSubobjectEditorTreeNodePtrType& InNodePtr)
+                {
+                    if(InNodePtr->GetComponentTemplate() == RootComp)
+                    {
+                    	RootHandle = InNodePtr->GetDataHandle();
+                    }
+                });
+			
+			return FindSlateNodeForHandle(RootHandle);
+		}
+	}
+
+	return nullptr;
+}
+
+void SSubobjectEditor::UpdateSelectionFromNodes(const TArray<FSubobjectEditorTreeNodePtrType>& SelectedNodes)
+{
+	bUpdatingSelection = true;
+
+	// Notify that the selection has updated
+	OnSelectionUpdated.ExecuteIfBound(SelectedNodes);
+
+	bUpdatingSelection = false;
+}
+
+void SSubobjectEditor::OnTreeSelectionChanged(FSubobjectEditorTreeNodePtrType InSelectedNodePtr, ESelectInfo::Type SelectInfo)
+{
+	UpdateSelectionFromNodes(TreeWidget->GetSelectedItems());
+}
+
+void SSubobjectEditor::HandleItemDoubleClicked(FSubobjectEditorTreeNodePtrType InItem)
+{
+	OnItemDoubleClicked.ExecuteIfBound(InItem);
+}
+
+void SSubobjectEditor::OnFindReferences()
+{
+	TArray<FSubobjectEditorTreeNodePtrType> SelectedNodes = TreeWidget->GetSelectedItems();
+	if (SelectedNodes.Num() == 1)
+	{
+		TSharedPtr<IToolkit> FoundAssetEditor = FToolkitManager::Get().FindEditorForAsset(GetBlueprint());
+		FSubobjectData* Data = SelectedNodes[0]->GetDataSource();
+		if (FoundAssetEditor.IsValid() && Data)
+		{
+			const FString VariableName = Data->GetVariableName().ToString();
+
+			// Search for both an explicit variable reference (finds get/sets of exactly that var, without including related-sounding variables)
+			// and a softer search for (VariableName) to capture bound component/widget event nodes which wouldn't otherwise show up
+			//@TODO: This logic is duplicated in SMyBlueprint::OnFindReference(), keep in sync
+			const FString SearchTerm = FString::Printf(
+				TEXT("Nodes(VariableReference(MemberName=+\"%s\") || Name=\"(%s)\")"), *VariableName, *VariableName);
+
+			TSharedRef<IBlueprintEditor> BlueprintEditor = StaticCastSharedRef<IBlueprintEditor>(
+				FoundAssetEditor.ToSharedRef());
+			BlueprintEditor->SummonSearchUI(true, SearchTerm);
+		}
+	}
+}
+
+void SSubobjectEditor::SelectRoot()
+{
+	if (RootNodes.Num() > 0)
+	{
+		TreeWidget->SetSelection(RootNodes[0]);
+	}
+}
+
+void SSubobjectEditor::SelectNodeFromHandle(const FSubobjectDataHandle& InHandle, bool bIsCntrlDown)
+{
+	SelectNode(FindSlateNodeForHandle(InHandle), bIsCntrlDown);
+}
+
+void SSubobjectEditor::SelectNode(FSubobjectEditorTreeNodePtrType InNodeToSelect, bool bIsCntrlDown)
+{
+	if (TreeWidget.IsValid() && InNodeToSelect.IsValid())
+	{
+		if (!bIsCntrlDown)
+		{
+			TreeWidget->SetSelection(InNodeToSelect);
+		}
+		else
+		{
+			TreeWidget->SetItemSelection(InNodeToSelect, !TreeWidget->IsItemSelected(InNodeToSelect));
+		}
+	}
+}
+
+EVisibility SSubobjectEditor::GetComponentsFilterBoxVisibility() const
+{
+	return (UICustomization.IsValid() && UICustomization->HideComponentsFilterBox())
+		       ? EVisibility::Collapsed
+		       : EVisibility::Visible;
+}
+
+void SSubobjectEditor::OnFilterTextChanged(const FText& InFilterText)
+{
+	struct OnFilterTextChanged_Inner
+	{
+		static FSubobjectEditorTreeNodePtrType ExpandToFilteredChildren(SSubobjectEditor* SubobjectEditor,
+		                                                                FSubobjectEditorTreeNodePtrType TreeNode)
+		{
+			FSubobjectEditorTreeNodePtrType NodeToFocus;
+
+			const TArray<FSubobjectEditorTreeNodePtrType>& Children = TreeNode->GetChildren();
+			// iterate backwards so we select from the top down
+			for (int32 ChildIndex = Children.Num() - 1; ChildIndex >= 0; --ChildIndex)
+			{
+				const FSubobjectEditorTreeNodePtrType& Child = Children[ChildIndex];
+				// Don't attempt to focus a separator or filtered node
+				if (!Child->IsSeperator() && !Child->IsFlaggedForFiltration())
+				{
+					SubobjectEditor->SetNodeExpansionState(TreeNode, /*bIsExpanded =*/true);
+					NodeToFocus = ExpandToFilteredChildren(SubobjectEditor, Child);
+				}
+			}
+
+			if (!NodeToFocus.IsValid() && !TreeNode->IsFlaggedForFiltration())
+			{
+				NodeToFocus = TreeNode;
+			}
+			return NodeToFocus;
+		}
+	};
+
+	FSubobjectEditorTreeNodePtrType NewSelection;
+	// iterate backwards so we select from the top down
+	for (int32 ComponentIndex = RootNodes.Num() - 1; ComponentIndex >= 0; --ComponentIndex)
+	{
+		FSubobjectEditorTreeNodePtrType Node = RootNodes[ComponentIndex];
+
+		bool bIsRootVisible = !RefreshFilteredState(Node, true);
+		TreeWidget->SetItemExpansion(Node, bIsRootVisible);
+		if (bIsRootVisible)
+		{
+			if (!GetFilterText().IsEmpty())
+			{
+				NewSelection = OnFilterTextChanged_Inner::ExpandToFilteredChildren(this, Node);
+			}
+		}
+	}
+
+	if (!NewSelection.IsValid() && RootNodes.Num() > 0)
+	{
+		NewSelection = RootNodes[0];
+	}
+
+	if (NewSelection.IsValid() && !TreeWidget->IsItemSelected(NewSelection))
+	{
+		SelectNode(NewSelection, /*IsCntrlDown =*/false);
+	}
+
+	UpdateTree(/*bRegenerateTreeNodes =*/false);
+}
+
+void SSubobjectEditor::OnItemScrolledIntoView(FSubobjectEditorTreeNodePtrType InItem, const TSharedPtr<ITableRow>& InWidget)
+{
+	if(DeferredRenameRequest.IsValid())
+	{
+		const FSubobjectDataHandle& ItemHandle = InItem->GetDataHandle();
+		if(DeferredRenameRequest == ItemHandle)
+		{
+			DeferredRenameRequest = FSubobjectDataHandle::InvalidHandle;
+			
+			// Transfer responsibility to end the 'create + give initial name' transaction to the tree item if such transaction is ongoing.
+			// We want to request a rename from the inline textbox on the associated handle
+			InItem->SetOngoingCreateTransaction(MoveTemp(DeferredOngoingCreateTransaction));
+			// Notify any listeners that this node will be renamed (this will trigger the inline text boxes on the
+			// row widgets to turn to edit mode)
+			InItem->GetRenameRequestedDelegate().ExecuteIfBound();
+		}
+	}
+}
+
+void SSubobjectEditor::UpdateTree(bool bRegenerateTreeNodes /* = true */)
+{
+	check(TreeWidget.IsValid());
+
+	// Early exit if we're deferring tree updates
+	if (!bAllowTreeUpdates)
+	{
+		return;
+	}
+
+	UObject* Context = GetObjectContext();
+
+	if (bRegenerateTreeNodes && Context)
+	{
+		// Obtain the set of expandable tree nodes that are currently collapsed
+		TSet<FSubobjectEditorTreeNodePtrType> CollapsedTreeNodes;
+		GetCollapsedNodes(RootNodes.Num() > 0 ? RootNodes[0] : FSubobjectEditorTreeNodePtrType(), CollapsedTreeNodes);
+
+		// Obtain the list of selected items
+		TArray<FSubobjectEditorTreeNodePtrType> SelectedTreeNodes = GetSelectedNodes();
+
+		// Clear the current tree
+		if (SelectedTreeNodes.Num() != 0)
+		{
+			TreeWidget->ClearSelection();
+		}
+		RootNodes.Empty();
+
+		ensureMsgf(FModuleManager::Get().LoadModule("SubobjectDataInterface"), TEXT("The Subobject Data Interface module is required."));
+
+		USubobjectDataSubsystem* DataSubsystem = USubobjectDataSubsystem::Get();
+		check(DataSubsystem);
+
+		TArray<FSubobjectDataHandle> SubobjectData;
+				
+		DataSubsystem->GatherSubobjectData(Context, SubobjectData);
+		
+		FSubobjectEditorTreeNodePtrType SeperatorNode;
+		// By default, root node will always be expanded. If possible we will restore the collapsed state later.
+		if (SubobjectData.Num() > 0)
+		{
+			FSubobjectEditorTreeNodePtrType Node = MakeShareable<FSubobjectEditorTreeNode>(
+				new FSubobjectEditorTreeNode(SubobjectData[0]));
+			RootNodes.Add(Node);
+			CachedRootHandle = Node->GetDataHandle();
+
+			SeperatorNode = MakeShareable<FSubobjectEditorTreeNode>(
+				new FSubobjectEditorTreeNode(FSubobjectDataHandle::InvalidHandle, /** bIsSeperator */true));
+			RootNodes.Add(SeperatorNode);
+			
+			TreeWidget->SetItemExpansion(Node, true);
+			TreeWidget->SetItemExpansion(SeperatorNode, true);
+
+			RefreshFilteredState(Node, false);
+		}
+
+		// Create slate nodes for each subobject
+		for (const FSubobjectDataHandle& Handle : SubobjectData)
+		{
+			// Do we have a slate node for this handle already? If not, then we need to make one
+			FSubobjectEditorTreeNodePtrType NewNode = FindOrCreateSlateNodeForHandle(Handle);
+			FSubobjectData* Data = Handle.GetData();
+
+			const FSubobjectDataHandle& ParentHandle = Data->GetParentHandle();
+
+			// Have parent? 
+			if (ParentHandle.IsValid())
+			{
+				// Get the parent node for this subobject
+				FSubobjectEditorTreeNodePtrType ParentNode = FindOrCreateSlateNodeForHandle(ParentHandle);
+				check(ParentNode);
+				ParentNode->AddChild(NewNode);
+				TreeWidget->SetItemExpansion(ParentNode, true);			
+			}
+			TreeWidget->SetItemExpansion(NewNode, true);
+		}
+
+		RestoreSelectionState(SelectedTreeNodes);
+
+		// If we have a pending deferred rename request, redirect it to the new tree node
+		if(DeferredRenameRequest.IsValid())
+		{
+			FSubobjectEditorTreeNodePtrType NodeToRenamePtr = FindSlateNodeForHandle(DeferredRenameRequest);
+			if(NodeToRenamePtr.IsValid())
+			{
+				TreeWidget->RequestScrollIntoView(NodeToRenamePtr);
+			}
+		}
+	}
+	
+	TreeWidget->RequestTreeRefresh();
+}
+
+void SSubobjectEditor::RefreshSelectionDetails()
+{
+	UpdateSelectionFromNodes(TreeWidget->GetSelectedItems());
+}
+
+void SSubobjectEditor::ClearSelection()
+{
+	if (bUpdatingSelection == false)
+	{
+		check(TreeWidget.IsValid());
+		TreeWidget->ClearSelection();
+	}
+}
+
+void SSubobjectEditor::RestoreSelectionState(TArray<FSubobjectEditorTreeNodePtrType>& SelectedTreeNodes)
+{
+	if (SelectedTreeNodes.Num() > 0)
+	{
+		// If there is only one item selected, imitate user selection to preserve navigation
+		ESelectInfo::Type SelectInfo = SelectedTreeNodes.Num() == 1 ? ESelectInfo::OnMouseClick : ESelectInfo::Direct;
+
+		// Restore the previous selection state on the new tree nodes
+		for (int i = 0; i < SelectedTreeNodes.Num(); ++i)
+		{
+			if (RootNodes.Num() > 0 && SelectedTreeNodes[i] == RootNodes[0])
+			{
+				TreeWidget->SetItemSelection(RootNodes[0], true, SelectInfo);
+			}
+			else
+			{
+				FSubobjectEditorTreeNodePtrType NodeToSelectPtr = FindSlateNodeForHandle(SelectedTreeNodes[i]->GetDataHandle());
+				if (NodeToSelectPtr.IsValid())
+				{
+					TreeWidget->SetItemSelection(NodeToSelectPtr, true, SelectInfo);
+				}
+			}
+
+			//if (GetEditorMode() != EComponentEditorMode::BlueprintSCS)
+			//{
+			//	TArray<FSubobjectEditorTreeNodePtrType> NewSelectedTreeNodes = TreeWidget->GetSelectedItems();
+			//	if (NewSelectedTreeNodes.Num() == 0)
+			//	{
+			//		TreeWidget->SetItemSelection(RootNodes[0], true);
+			//	}
+			//}
+		}
+	}
+}
+
+FReply SSubobjectEditor::OnKeyDown(const FGeometry& MyGeometry, const FKeyEvent& InKeyEvent)
+{
+	if (CommandList->ProcessCommandBindings(InKeyEvent))
+	{
+		return FReply::Handled();
+	}
+	return FReply::Unhandled();
+}
+
+FReply SSubobjectEditor::TryHandleAssetDragDropOperation(const FDragDropEvent& DragDropEvent)
+{
+	// #TODO_BH Implement
+	ensureMsgf(false, TEXT("Implement TryHandleAssetDragDropOperation!"));
+	return FReply::Unhandled();
+}
+
+void SSubobjectEditor::SetUICustomization(TSharedPtr<ISCSEditorUICustomization> InUICustomization)
+{
+	UICustomization = InUICustomization;
+
+	UpdateTree(true /*bRegenerateTreeNodes*/);
+}
+
+bool SSubobjectEditor::CanCutNodes() const
+{
+	return CanCopyNodes() && CanDeleteNodes();
+}
+
+void SSubobjectEditor::CutSelectedNodes()
+{
+	TArray<FSubobjectEditorTreeNodePtrType> SelectedNodes = GetSelectedNodes();
+	const FScopedTransaction Transaction(SelectedNodes.Num() > 1
+		                                     ? LOCTEXT("CutComponents", "Cut Components")
+		                                     : LOCTEXT("CutComponent", "Cut Component"));
+
+	FSubobjectEditorTreeNodePtrType PostCutSelection = RootNodes.Num() > 0 ? RootNodes[0] : nullptr;
+	for(FSubobjectEditorTreeNodePtrType& SelectedNode : SelectedNodes)
+	{
+		if(FSubobjectEditorTreeNodePtrType ParentNode = SelectedNode->GetParent())
+		{
+			if(!SelectedNodes.Contains(ParentNode))
+			{
+				PostCutSelection = ParentNode;
+				break;
+			}
+		}
+	}
+	
+	CopySelectedNodes();
+	OnDeleteNodes();
+
+	// Select the parent node of the thing that was just cut
+	if(PostCutSelection.IsValid())
+	{
+		TreeWidget->SetSelection(PostCutSelection);
+	}
+}
+
+bool SSubobjectEditor::CanCopyNodes() const
+{
+	TArray<FSubobjectDataHandle> SelectedHandles = GetSelectedHandles();
+	if(USubobjectDataSubsystem* System = USubobjectDataSubsystem::Get())
+	{
+		return System->CanCopySubobjects(SelectedHandles);
+	}
+	return false;
+}
+
+bool SSubobjectEditor::CanPasteNodes() const
+{
+	if(!IsEditingAllowed())
+	{
+		return false;
+	}
+	
+	if(FSubobjectEditorTreeNodePtrType SceneRoot = GetSceneRootNode())
+	{
+		if(USubobjectDataSubsystem* System = USubobjectDataSubsystem::Get())
+		{
+			return System->CanPasteSubobjects(SceneRoot->GetDataHandle());
+		}
+	}
+
+	return false;
+}
+
+bool SSubobjectEditor::CanDeleteNodes() const
+{
+	if(!IsEditingAllowed())
+	{
+		return false;
+	}
+
+	TArray<FSubobjectEditorTreeNodePtrType> SelectedNodes = GetSelectedNodes();
+	for (const FSubobjectEditorTreeNodePtrType& SelectedNode : SelectedNodes)
+	{
+		if (!SelectedNode->IsValid() || !SelectedNode->CanDelete())
+		{
+			return false;
+		}
+		else
+		{
+			// Don't allow nodes that belong to a child actor template to be deleted
+			const bool bIsChildActorSubtreeNode = (SelectedNode->GetDataSource()->IsChildActor());
+			if (bIsChildActorSubtreeNode)
+			{
+				return false;
+			}
+		}
+	}
+	
+	return true;
+}
+
+bool SSubobjectEditor::CanDuplicateComponent() const
+{
+	if(!IsEditingAllowed())
+	{
+		return false;
+	}
+	
+	return CanCopyNodes();
+}
+
+bool SSubobjectEditor::CanRenameComponent() const
+{
+	if (!IsEditingAllowed())
+	{
+		return false;
+	}
+	// In addition to certain node types, don't allow nodes within a child actor template's hierarchy to be renamed
+	TArray<FSubobjectEditorTreeNodePtrType> SelectedItems = TreeWidget->GetSelectedItems();
+	
+	return SelectedItems.Num() == 1 && SelectedItems[0]->CanRename();
+}
+
+void SSubobjectEditor::OnRenameComponent()
+{
+	OnRenameComponent(nullptr);
+}
+
+void SSubobjectEditor::OnRenameComponent(TUniquePtr<FScopedTransaction> InComponentCreateTransaction)
+{
+	TArray<FSubobjectEditorTreeNodePtrType> SelectedNodes = GetSelectedNodes();
+
+	// Should already be prevented from making it here.
+	check(SelectedNodes.Num() == 1);
+
+	DeferredRenameRequest = SelectedNodes[0]->GetDataHandle();
+
+	// If this fails, something in the chain of responsibility failed to end the previous transaction.
+	check(!DeferredOngoingCreateTransaction.IsValid());
+	// If a 'create + give initial name' transaction is ongoing, take responsibility of ending it until the selected item is scrolled into view.
+	DeferredOngoingCreateTransaction = MoveTemp(InComponentCreateTransaction); 
+	
+	TreeWidget->RequestScrollIntoView(SelectedNodes[0]);
+
+	if (DeferredOngoingCreateTransaction.IsValid() && !PostTickHandle.IsValid())
+	{
+		// Ensure the item will be scrolled into view during the frame (See explanation in OnPostTick()).
+		PostTickHandle = FSlateApplication::Get().OnPostTick().AddSP(this, &SSubobjectEditor::OnPostTick);
+	}
+}
+
+void SSubobjectEditor::OnPostTick(float)
+{
+	// If a 'create + give initial name' is ongoing and the transaction ownership was not transferred during the frame it was requested, it is most likely because the newly
+	// created item could not be scrolled into view (should say 'teleported', the scrolling is not animated). The tree view will not put the item in view if the there is
+	// no space left to display the item. (ex a splitter where all the display space is used by the other component). End the transaction before starting a new frame. The user
+	// will not be able to rename on creation, the widget is likely not in view and cannot be edited anyway.
+	DeferredOngoingCreateTransaction.Reset();
+
+	// The post tick event handler is not required anymore.
+	FSlateApplication::Get().OnPostTick().Remove(PostTickHandle);
+	PostTickHandle.Reset();
+}
+
+UObject* SSubobjectEditor::GetObjectContext() const
+{
+	return ObjectContext.Get(nullptr);
+}
+
+FSubobjectDataHandle SSubobjectEditor::GetObjectContextHandle() const
+{
+	if (RootNodes.Num() > 0)
+	{
+		ensure(RootNodes[0]->GetObject() == GetObjectContext());
+		return RootNodes[0]->GetDataHandle();
+	}
+	else
+	{
+		return FSubobjectDataHandle::InvalidHandle;
+	}
+}
+
+UBlueprint* SSubobjectEditor::GetBlueprint() const
+{
+	if (AActor* Actor = Cast<AActor>(GetObjectContext()))
+	{
+		const UClass* ActorClass = Actor->GetClass();
+		check(ActorClass != nullptr);
+
+		return UBlueprint::GetBlueprintFromClass(ActorClass);
+	}
+
+	return nullptr;
+}
+
+void SSubobjectEditor::DepthFirstTraversal(const FSubobjectEditorTreeNodePtrType& InNodePtr, TSet<FSubobjectEditorTreeNodePtrType>& OutVisitedNodes, const TFunctionRef<void(const FSubobjectEditorTreeNodePtrType&)> InFunction) const
+{
+	if (InNodePtr.IsValid()
+		&& ensureMsgf(!OutVisitedNodes.Contains(InNodePtr), TEXT("Already visited node: %s (Parent: %s)"),
+		              *InNodePtr->GetDisplayString(),
+		              InNodePtr->GetParent().IsValid() ? *InNodePtr->GetParent()->GetDisplayString() : TEXT("NULL")))
+	{
+		InFunction(InNodePtr);
+		OutVisitedNodes.Add(InNodePtr);
+
+		for (const FSubobjectEditorTreeNodePtrType& Child : InNodePtr->GetChildren())
+		{
+			DepthFirstTraversal(Child, OutVisitedNodes, InFunction);
+		}
+	}
+}
+
+void SSubobjectEditor::GetCollapsedNodes(const FSubobjectEditorTreeNodePtrType& InNodePtr, TSet<FSubobjectEditorTreeNodePtrType>& OutCollapsedNodes) const
+{
+	TSet<FSubobjectEditorTreeNodePtrType> VisitedNodes;
+	DepthFirstTraversal(InNodePtr, VisitedNodes,
+        [TreeWidget = this->TreeWidget, &OutCollapsedNodes](
+        const FSubobjectEditorTreeNodePtrType& InNodePtr)
+        {
+            if (InNodePtr->GetChildren().Num() > 0 && !TreeWidget->IsItemExpanded(InNodePtr))
+            {
+                OutCollapsedNodes.Add(InNodePtr);
+            }
+        });
+}
+
+FSubobjectEditorTreeNodePtrType SSubobjectEditor::FindSlateNodeForHandle(const FSubobjectDataHandle& InHandleToFind, FSubobjectEditorTreeNodePtrType InStartNodePtr) const
+{
+	FSubobjectEditorTreeNodePtrType OutNodePtr;
+	if (InHandleToFind.IsValid() && RootNodes.Num() > 0)
+	{
+		FSubobjectData* DataToFind = InHandleToFind.GetSharedDataPtr().Get();
+		
+		TSet<FSubobjectEditorTreeNodePtrType> VisitedNodes;
+		DepthFirstTraversal(RootNodes[0], VisitedNodes,
+            [&OutNodePtr, &DataToFind](
+            const FSubobjectEditorTreeNodePtrType& CurNodePtr)
+            {
+            	if(CurNodePtr->GetDataHandle().IsValid())
+            	{
+            		if(CurNodePtr->GetDataHandle().GetSharedDataPtr()->GetObject() == DataToFind->GetObject())
+            		{
+                        OutNodePtr = CurNodePtr;
+                    }
+            	}
+            });
+	}
+
+	return OutNodePtr;
+}
+
+FSubobjectEditorTreeNodePtrType SSubobjectEditor::FindSlateNodeForObject(const UObject* InObject, bool bIncludeAttachmentComponents) const
+{
+	FSubobjectEditorTreeNodePtrType OutNodePtr;
+
+	if (RootNodes.Num() > 0)
+	{
+		TSet<FSubobjectEditorTreeNodePtrType> VisitedNodes;
+		DepthFirstTraversal(RootNodes[0], VisitedNodes,
+            [&OutNodePtr, &InObject](
+            const FSubobjectEditorTreeNodePtrType& CurNodePtr)
+            {
+                if(CurNodePtr->GetDataHandle().IsValid())
+                {
+                    if(CurNodePtr->GetDataHandle().GetSharedDataPtr()->GetObject() == InObject)
+                    {
+                        OutNodePtr = CurNodePtr;
+                    }
+                }
+            });
+	}
+
+	return OutNodePtr;
+}
+
+FSubobjectEditorTreeNodePtrType SSubobjectEditor::FindOrCreateSlateNodeForHandle(const FSubobjectDataHandle& Handle) const
+{
+	if (FSubobjectEditorTreeNodePtrType Existing = FindSlateNodeForHandle(Handle))
+	{
+		return Existing;
+	}
+	else
+	{
+		return MakeShareable<FSubobjectEditorTreeNode>(new FSubobjectEditorTreeNode(Handle));
+	}
+}
+
+void SSubobjectEditor::SetNodeExpansionState(FSubobjectEditorTreeNodePtrType InNodeToChange, const bool bIsExpanded)
+{
+	if (TreeWidget.IsValid() && InNodeToChange.IsValid())
+	{
+		TreeWidget->SetItemExpansion(InNodeToChange, bIsExpanded);
+	}
+}
+
+void SSubobjectEditor::SetItemExpansionRecursive(FSubobjectEditorTreeNodePtrType Model, bool bInExpansionState)
+{
+	SetNodeExpansionState(Model, bInExpansionState);
+	for (const FSubobjectEditorTreeNodePtrType& Child : Model->GetChildren())
+	{
+		if (Child.IsValid())
+		{
+			SetItemExpansionRecursive(Child, bInExpansionState);
+		}
+	}
+}
+
+void SSubobjectEditor::DumpTree()
+{
+	/* Example:
+
+		[ACTOR] MyBlueprint (self)
+		|
+		[SEPARATOR]
+		|
+		DefaultSceneRoot (Inherited)
+		|
+		+- StaticMesh (Inherited)
+		|  |
+		|  +- Scene4 (Inherited)
+		|  |
+		|  +- Scene (Inherited)
+		|     |
+		|     +- Scene1 (Inherited)
+		|
+		+- Scene2 (Inherited)
+		|  |
+		|  +- Scene3 (Inherited)
+		|
+		[SEPARATOR]
+		|
+		ProjectileMovement (Inherited)
+	*/
+
+	UE_LOG(LogSubobjectEditor, Log, TEXT("---------------------"));
+	UE_LOG(LogSubobjectEditor, Log, TEXT(" STreeView NODE DUMP"));
+	UE_LOG(LogSubobjectEditor, Log, TEXT("---------------------"));
+
+	//TSet<FSubobjectEditorTreeNodePtrType> VisitedNodes;
+	//DepthFirstTraversal(RootNodes[0], VisitedNodes, [](const FSubobjectEditorTreeNodePtrType& InNode)
+	//{
+	//	UE_LOG(LogSubobjectEditor, Log, TEXT(""));
+	//});
+}
+
+bool SSubobjectEditor::IsEditingAllowed() const
+{
+	return AllowEditing.Get() && nullptr == GEditor->PlayWorld;
+}
+
+TSharedPtr<SWidget> SSubobjectEditor::GetToolButtonsBox() const
+{
+	return ButtonBox;
+}
+
+void SSubobjectEditor::RegisterContextMenu()
+{
+	UToolMenus* ToolMenus = UToolMenus::Get();
+	if (!ToolMenus->IsMenuRegistered(SubobjectTree_ContextMenuName))
+	{
+		UToolMenu* Menu = ToolMenus->RegisterMenu(SubobjectTree_ContextMenuName);
+		Menu->AddDynamicSection("SCSEditorDynamic", FNewToolMenuDelegate::CreateLambda([](UToolMenu* InMenu)
+		{
+			USSCSEditorMenuContext* ContextObject = InMenu->FindContext<USSCSEditorMenuContext>();
+			if (ContextObject && ContextObject->SubobjectEditor.IsValid())
+			{
+				ContextObject->SubobjectEditor.Pin()->PopulateContextMenu(InMenu);
+			}
+		}));
+	}
+}
+
+TSharedPtr<SWidget> SSubobjectEditor::CreateContextMenu()
+{
+	TArray<FSubobjectEditorTreeNodePtrType> SelectedItems = TreeWidget->GetSelectedItems();
+
+	if (SelectedItems.Num() > 0 || CanPasteNodes())
+	{
+		RegisterContextMenu();
+		USSCSEditorMenuContext* ContextObject = NewObject<USSCSEditorMenuContext>();
+		ContextObject->SubobjectEditor = SharedThis(this);
+		FToolMenuContext ToolMenuContext(CommandList, TSharedPtr<FExtender>(), ContextObject);
+		return UToolMenus::Get()->GenerateWidget(SubobjectTree_ContextMenuName, ToolMenuContext);
+	}
+	return TSharedPtr<SWidget>();
+}
+
+void SSubobjectEditor::PopulateContextMenu(UToolMenu* Menu)
+{
+	TArray<FSubobjectEditorTreeNodePtrType> SelectedItems = TreeWidget->GetSelectedItems();
+
+	if(SelectedItems.Num() > 0 || CanPasteNodes())
+	{
+		bool bOnlyShowPasteOption = false;
+
+		if(SelectedItems.Num() > 0)
+		{
+			FSubobjectData* SelectedData = SelectedItems[0]->GetDataSource();
+			const AActor* SelectedActor =  SelectedData->GetObject<AActor>();
+			if (SelectedItems.Num() == 1 && SelectedActor)
+			{
+				if (SelectedData->IsChildActor())
+				{
+					// Include specific context menu options for a single child actor node selection
+					SubobjectEditorUtils::FillChildActorContextMenuOptions(Menu, SelectedItems[0]);
+				}
+				else
+				{
+					bOnlyShowPasteOption = true;
+				}
+			}
+			else
+			{
+				for (const FSubobjectEditorTreeNodePtrType& SelectedNode : SelectedItems)
+				{
+					if (!SelectedNode->IsComponentNode())
+					{
+						bOnlyShowPasteOption = true;
+						break;
+					}
+				}
+
+				if (!bOnlyShowPasteOption)
+				{
+					bool bIsChildActorSubtreeNodeSelected = false;
+					
+					TArray<UActorComponent*> SelectedComponents;
+					for(const FSubobjectEditorTreeNodePtrType& SelectedNodePtr : SelectedItems)
+					{
+						check(SelectedNodePtr->IsValid());
+						// Get the component template associated with the selected node
+						const UActorComponent* ComponentTemplate = SelectedNodePtr->GetComponentTemplate();
+						if (ComponentTemplate)
+						{
+							// #TODO_BH Remove this const cast
+							SelectedComponents.Add(const_cast<UActorComponent*>(ComponentTemplate));
+						}
+
+						// Determine if any selected node belongs to a child actor template
+						if (!bIsChildActorSubtreeNodeSelected)
+						{
+							bIsChildActorSubtreeNodeSelected = SelectedNodePtr->IsChildSubtreeNode();
+						}
+					}
+
+					PopulateContextMenuImpl(Menu, SelectedItems, bIsChildActorSubtreeNodeSelected);
+				}
+			}
+		}
+		else
+		{
+			bOnlyShowPasteOption = true;
+		}
+
+		if (bOnlyShowPasteOption)
+		{
+			FToolMenuSection& Section = Menu->AddSection("PasteComponent", LOCTEXT("EditComponentHeading", "Edit"));
+			{
+				Section.AddMenuEntry(FGenericCommands::Get().Paste);
+			}
+		}
+	}
+}
+
+void SSubobjectEditor::GetSelectedItemsForContextMenu(TArray<FComponentEventConstructionData>& OutSelectedItems) const
+{
+	TArray<FSubobjectEditorTreeNodePtrType> SelectedTreeItems = TreeWidget->GetSelectedItems();
+	for (auto NodeIter = SelectedTreeItems.CreateConstIterator(); NodeIter; ++NodeIter)
+	{
+		FComponentEventConstructionData NewItem;
+		const FSubobjectEditorTreeNodePtrType& TreeNode = *NodeIter;
+		const FSubobjectData* Data = TreeNode->GetDataSource();
+		NewItem.VariableName = Data->GetVariableName();
+		NewItem.Component = const_cast<UActorComponent*>(TreeNode->GetComponentTemplate());
+		OutSelectedItems.Add(NewItem);
+	}
+}
+
+TSubclassOf<UActorComponent> SSubobjectEditor::GetComponentTypeFilterToApply() const
+{
+	TSubclassOf<UActorComponent> ComponentType = UICustomization.IsValid()
+		                                             ? UICustomization->GetComponentTypeFilter()
+		                                             : nullptr;
+	if (!ComponentType)
+	{
+		ComponentType = ComponentTypeFilter.Get();
+	}
+	return ComponentType;
+}
+
+bool SSubobjectEditor::RefreshFilteredState(FSubobjectEditorTreeNodePtrType TreeNode, bool bRecursive)
+{
+	const UClass* FilterType = GetComponentTypeFilterToApply();
+
+	FString FilterText = FText::TrimPrecedingAndTrailing(GetFilterText()).ToString();
+	TArray<FString> FilterTerms;
+	FilterText.ParseIntoArray(FilterTerms, TEXT(" "), /*CullEmpty =*/true);
+
+	TreeNode->RefreshFilteredState(FilterType, FilterTerms, bRecursive);
+	return TreeNode->IsFlaggedForFiltration();
+}
+
+FText SSubobjectEditor::GetFilterText() const
+{
+	return FilterBox->GetText();
+}
+
+void SSubobjectEditor::PromoteToBlueprint() const
+{
+	if (AActor* ActorContext = Cast<AActor>(GetObjectContext()))
+	{
+		FCreateBlueprintFromActorDialog::OpenDialog(ECreateBlueprintFromActorMode::Subclass, ActorContext);
+	}
+}
+
+FReply SSubobjectEditor::OnPromoteToBlueprintClicked()
+{
+	PromoteToBlueprint();
+	return FReply::Handled();
+}
+
+void SSubobjectEditor::OnOpenBlueprintEditor(bool bForceCodeEditing) const
+{
+	if (AActor* ActorInstance = Cast<AActor>(GetObjectContext()))
+	{
+		if (UBlueprint* Blueprint = Cast<UBlueprint>(ActorInstance->GetClass()->ClassGeneratedBy))
+		{
+			if (bForceCodeEditing && (Blueprint->UbergraphPages.Num() > 0))
+			{
+				FKismetEditorUtilities::BringKismetToFocusAttentionOnObject(Blueprint->GetLastEditedUberGraph());
+			}
+			else
+			{
+				GEditor->GetEditorSubsystem<UAssetEditorSubsystem>()->OpenEditorForAsset(Blueprint);
+			}
+		}
+	}
+}
+
+////////////////////////////////////////////////
+// Class filters for creating new components
+class FComponentClassParentFilter : public IClassViewerFilter
+{
+public:
+	FComponentClassParentFilter(const TSubclassOf<UActorComponent>& InComponentClass) : ComponentClass(InComponentClass)
+	{
+	}
+
+	virtual bool IsClassAllowed(const FClassViewerInitializationOptions& InInitOptions, const UClass* InClass, TSharedRef<FClassViewerFilterFuncs> InFilterFuncs) override
+	{
+		return InClass->IsChildOf(ComponentClass);
+	}
+
+	virtual bool IsUnloadedClassAllowed(const FClassViewerInitializationOptions& InInitOptions, const TSharedRef<const IUnloadedBlueprintData> InUnloadedClassData, TSharedRef<FClassViewerFilterFuncs> InFilterFuncs) override
+	{
+		return InUnloadedClassData->IsChildOf(ComponentClass);
+	}
+
+	TSubclassOf<UActorComponent> ComponentClass;
+};
+
+typedef FComponentClassParentFilter FNativeComponentClassParentFilter;
+
+class FBlueprintComponentClassParentFilter : public FComponentClassParentFilter
+{
+public:
+	FBlueprintComponentClassParentFilter(const TSubclassOf<UActorComponent>& InComponentClass)
+		: FComponentClassParentFilter(InComponentClass)
+	{}
+
+	virtual bool IsClassAllowed(const FClassViewerInitializationOptions& InInitOptions, const UClass* InClass, TSharedRef<FClassViewerFilterFuncs> InFilterFuncs) override
+	{
+		return
+			FComponentClassParentFilter::IsClassAllowed(InInitOptions, InClass, InFilterFuncs) &&
+			FKismetEditorUtilities::CanCreateBlueprintOfClass(InClass);
+	}
+};
+
+UClass* SSubobjectEditor::SpawnCreateNewCppComponentWindow(TSubclassOf<UActorComponent> ComponentClass)
+{
+	TSharedPtr<SWindow> ParentWindow = FSlateApplication::Get().FindWidgetWindow(SharedThis(this));
+
+	FString AddedClassName;
+	auto OnCodeAddedToProject = [&AddedClassName](const FString& ClassName, const FString& ClassPath,
+	                                              const FString& ModuleName)
+	{
+		if (!ClassName.IsEmpty() && !ClassPath.IsEmpty())
+		{
+			AddedClassName = FString::Printf(TEXT("/Script/%s.%s"), *ModuleName, *ClassName);
+		}
+	};
+
+	FGameProjectGenerationModule::Get().OpenAddCodeToProjectDialog(
+		FAddToProjectConfig()
+		.WindowTitle(LOCTEXT("AddNewC++Component", "Add C++ Component"))
+		.ParentWindow(ParentWindow)
+		.Modal()
+		.OnAddedToProject(FOnAddedToProject::CreateLambda(OnCodeAddedToProject))
+		.FeatureComponentClasses()
+		.AllowableParents(MakeShareable(new FNativeComponentClassParentFilter(ComponentClass)))
+		.DefaultClassPrefix(TEXT("New"))
+	);
+
+
+	return LoadClass<UActorComponent>(nullptr, *AddedClassName, nullptr, LOAD_None, nullptr);
+}
+
+UClass* SSubobjectEditor::SpawnCreateNewBPComponentWindow(TSubclassOf<UActorComponent> ComponentClass)
+{
+	UClass* NewClass = nullptr;
+
+	auto OnAddedToProject = [&](const FString& ClassName, const FString& PackagePath, const FString& ModuleName)
+	{
+		if (!ClassName.IsEmpty() && !PackagePath.IsEmpty())
+		{
+			if (UPackage* Package = FindPackage(nullptr, *PackagePath))
+			{
+				if (UBlueprint* NewBP = FindObjectFast<UBlueprint>(Package, *ClassName))
+				{
+					NewClass = NewBP->GeneratedClass;
+
+					TArray<UObject*> Objects;
+					Objects.Emplace(NewBP);
+					GEditor->SyncBrowserToObjects(Objects);
+
+					// Open the editor for the new blueprint
+					GEditor->GetEditorSubsystem<UAssetEditorSubsystem>()->OpenEditorForAsset(NewBP);
+				}
+			}
+		}
+	};
+
+	FGameProjectGenerationModule::Get().OpenAddBlueprintToProjectDialog(
+		FAddToProjectConfig()
+		.WindowTitle(LOCTEXT("AddNewBlueprintComponent", "Add Blueprint Component"))
+		.ParentWindow(FSlateApplication::Get().FindWidgetWindow(SharedThis(this)))
+		.Modal()
+		.AllowableParents(MakeShareable(new FBlueprintComponentClassParentFilter(ComponentClass)))
+		.FeatureComponentClasses()
+		.OnAddedToProject(FOnAddedToProject::CreateLambda(OnAddedToProject))
+		.DefaultClassPrefix(TEXT("New"))
+	);
+
+	return NewClass;
+}
+
+FSubobjectDataHandle SSubobjectEditor::PerformComboAddClass(TSubclassOf<UActorComponent> ComponentClass, EComponentCreateAction::Type ComponentCreateAction, UObject* AssetOverride)
+{
+	FSubobjectDataHandle NewComponentHandle = FSubobjectDataHandle::InvalidHandle;
+
+	// Display a class picked for which class to use for the new component
+	UClass* NewClass = ComponentClass;
+
+	// Display a class creation menu if the user wants to create a new class
+	if (ComponentCreateAction == EComponentCreateAction::CreateNewCPPClass)
+	{
+		NewClass = SpawnCreateNewCppComponentWindow(ComponentClass);
+	}
+	else if (ComponentCreateAction == EComponentCreateAction::CreateNewBlueprintClass)
+	{
+		NewClass = SpawnCreateNewBPComponentWindow(ComponentClass);
+	}
+
+	USubobjectDataSubsystem* SubobjectSystem = USubobjectDataSubsystem::Get();
+
+	if (NewClass && SubobjectSystem)
+	{
+		TUniquePtr<FScopedTransaction> AddTransaction = MakeUnique<FScopedTransaction>(LOCTEXT("AddComponent", "Add Component"));
+
+		FEditorDelegates::LoadSelectedAssetsIfNeeded.Broadcast();
+		USelection* Selection = GEditor->GetSelectedObjects();
+
+		bool bAddedComponent = false;
+		FText OutFailReason;
+
+		// This adds components according to the type selected in the drop down. If the user
+		// has the appropriate objects selected in the content browser then those are added,
+		// else we go down the previous route of adding components by type.
+		
+		// Furthermore don't try to match up assets for USceneComponent it will match lots of things and doesn't have any nice behavior for asset adds 
+		if (Selection->Num() > 0 && !AssetOverride && NewClass != USceneComponent::StaticClass())
+		{
+			for (FSelectionIterator ObjectIter(*Selection); ObjectIter; ++ObjectIter)
+			{
+				UObject* Object = *ObjectIter;
+				UClass* Class = Object->GetClass();
+
+				TArray<TSubclassOf<UActorComponent>> ComponentClasses = FComponentAssetBrokerage::GetComponentsForAsset(Object);
+
+				// if the selected asset supports the selected component type then go ahead and add it
+				for (int32 ComponentIndex = 0; ComponentIndex < ComponentClasses.Num(); ComponentIndex++)
+				{
+					if (ComponentClasses[ComponentIndex]->IsChildOf(NewClass))
+					{
+						const FSubobjectDataHandle ParentHandle = SubobjectSystem->FindHandleForObject(CachedRootHandle, Object);
+						if(ParentHandle.IsValid())
+						{
+							NewComponentHandle = AddNewSubobject(ParentHandle, NewClass, Object, OutFailReason, MoveTemp(AddTransaction));
+							bAddedComponent = true;
+							break;
+						}
+					}
+				}
+			}
+		}
+
+		if (!bAddedComponent)
+		{
+			// Attach this component to the override asset first, but if none is given then use the actor context			
+			FSubobjectDataHandle ParentHandle = SubobjectSystem->FindHandleForObject(CachedRootHandle, AssetOverride);
+			
+			if(!ParentHandle.IsValid())
+			{
+				// If we have something selected, then we should attach it to that
+				TArray<FSubobjectEditorTreeNodePtrType> SelectedTreeNodes = TreeWidget->GetSelectedItems();
+				if(SelectedTreeNodes.Num() > 0)
+				{
+					ParentHandle = SelectedTreeNodes[0]->GetDataHandle();
+				}
+				// Otherwise fall back to the root node
+				else
+				{
+					ParentHandle = RootNodes.Num() > 0
+						               ? RootNodes[0]->GetDataHandle()
+						               : FSubobjectDataHandle::InvalidHandle;					
+				}
+			}
+			
+			if(ParentHandle.IsValid())
+			{
+				NewComponentHandle = AddNewSubobject(ParentHandle, NewClass, AssetOverride, OutFailReason, MoveTemp(AddTransaction));
+			}
+		}
+
+		// We failed adding a new component, display why with a slate notif
+		if(!NewComponentHandle.IsValid())
+		{
+			if(OutFailReason.IsEmpty())
+			{
+				OutFailReason = LOCTEXT("AddComponentFailed_Generic", "Failed to add component!");
+			}
+			
+			FNotificationInfo Info(OutFailReason);
+			Info.Image = FEditorStyle::GetBrush(TEXT("Icons.Error"));
+			Info.bFireAndForget = true;
+			Info.bUseSuccessFailIcons = false;
+			Info.ExpireDuration = 5.0f;
+
+			FSlateNotificationManager::Get().AddNotification(Info);
+		}
+		else
+		{
+			UpdateTree();
+			
+			// Set focus to the newly created subobject
+			FSubobjectEditorTreeNodePtrType NewNode = FindSlateNodeForHandle(NewComponentHandle);
+			if(NewNode != nullptr)
+			{
+				TreeWidget->SetSelection(NewNode);
+				OnRenameComponent(MoveTemp(AddTransaction));
+			}
+		}
+	}
+
+	return NewComponentHandle;
+}
+
+EVisibility SSubobjectEditor::GetPromoteToBlueprintButtonVisibility() const
+{
+	return (UICustomization.IsValid() && UICustomization->HideBlueprintButtons())
+	       || (ShowInlineSearchWithButtons())
+	       || (GetBlueprint() != nullptr)
+		       ? EVisibility::Collapsed
+		       : EVisibility::Visible;
+}
+
+EVisibility SSubobjectEditor::GetEditBlueprintButtonVisibility() const
+{
+	return (UICustomization.IsValid() && UICustomization->HideBlueprintButtons())
+	       || (ShowInlineSearchWithButtons())
+	       || (GetBlueprint() == nullptr)
+		       ? EVisibility::Collapsed
+		       : EVisibility::Visible;
+}
+
+void SSubobjectEditor::Utils::PopulateHandlesArray(const TArray<FSubobjectEditorTreeNodePtrType>& SlateNodePtrs, TArray<FSubobjectDataHandle>& OutHandles)
+{
+	for(const FSubobjectEditorTreeNodePtrType& DroppedNodePtr : SlateNodePtrs)
+    {
+    	OutHandles.Add(DroppedNodePtr->GetDataHandle());
+    }
+}
+
+#undef LOCTEXT_NAMESPACE
diff --git a/Engine/Source/Editor/SubobjectEditor/Private/SSubobjectEditorModule.cpp b/Engine/Source/Editor/SubobjectEditor/Private/SSubobjectEditorModule.cpp
new file mode 100644
--- /dev/null
+++ b/Engine/Source/Editor/SubobjectEditor/Private/SSubobjectEditorModule.cpp
@@ -0,0 +1,5 @@
+// Copyright Epic Games, Inc. All Rights Reserved.
+
+#include "SSubobjectEditorModule.h"
+
+IMPLEMENT_MODULE(FSubobjectEditorModule, SubobjectEditor)
\ No newline at end of file
diff --git a/Engine/Source/Editor/SubobjectEditor/Private/SSubobjectInstanceEditor.cpp b/Engine/Source/Editor/SubobjectEditor/Private/SSubobjectInstanceEditor.cpp
new file mode 100644
--- /dev/null
+++ b/Engine/Source/Editor/SubobjectEditor/Private/SSubobjectInstanceEditor.cpp
@@ -0,0 +1,476 @@
+// Copyright Epic Games, Inc. All Rights Reserved.
+
+#include "SSubobjectInstanceEditor.h"
+#include "ScopedTransaction.h"
+#include "IDocumentation.h"
+#include "GraphEditorActions.h"
+#include "Editor/EditorEngine.h"
+#include "ISCSEditorUICustomization.h"	// #TODO_BH Rename this to subobject
+#include "SCSEditorExtensionContext.h"	// #TODO_BH Rename this to subobject
+
+#include "Styling/SlateIconFinder.h"
+#include "SlateOptMacros.h"
+#include "Widgets/Input/SSearchBox.h"
+#include "Framework/MultiBox/MultiBoxBuilder.h"
+#include "SComponentClassCombo.h"
+#include "SEditorHeaderButton.h"
+#include "Editor/UnrealEdEngine.h"
+#include "Subsystems/PanelExtensionSubsystem.h"	// SExtensionPanel
+#include "Kismet2/ComponentEditorUtils.h"
+
+#define LOCTEXT_NAMESPACE "SSubobjectInstanceEditor"
+
+extern UNREALED_API UUnrealEdEngine* GUnrealEd;
+
+////////////////////////////////////////////////
+// SSubobjectInstanceEditor
+
+BEGIN_SLATE_FUNCTION_BUILD_OPTIMIZATION
+void SSubobjectInstanceEditor::Construct(const FArguments& InArgs)
+{
+	ObjectContext = InArgs._ObjectContext;
+	OnSelectionUpdated = InArgs._OnSelectionUpdated;
+	OnItemDoubleClicked = InArgs._OnItemDoubleClicked;
+	AllowEditing = InArgs._AllowEditing;
+	OnObjectReplaced = InArgs._OnObjectReplaced;
+	bAllowTreeUpdates = true;
+	
+	CreateCommandList();
+
+	// Build the tree widget
+	FSlateBrush const* MobilityHeaderBrush = FEditorStyle::GetBrush(TEXT("ClassIcon.ComponentMobilityHeaderIcon"));
+	
+	ConstructTreeWidget();
+
+	// Should only be true when used in the blueprints details panel
+	const bool bInlineSearchBarWithButtons = ShowInlineSearchWithButtons();
+
+	TSharedPtr<SWidget> Contents;
+
+	TSharedPtr<SVerticalBox> HeaderBox;
+	TSharedPtr<SWidget> SearchBar = SAssignNew(FilterBox, SSearchBox)
+		.HintText(!bInlineSearchBarWithButtons ? LOCTEXT("SearchComponentsHint", "Search Components") : LOCTEXT("SearchHint", "Search"))
+		.OnTextChanged(this, &SSubobjectInstanceEditor::OnFilterTextChanged)
+		.Visibility(this, &SSubobjectInstanceEditor::GetComponentsFilterBoxVisibility);
+
+	FMenuBuilder EditBlueprintMenuBuilder = CreateMenuBuilder();
+
+	// Extension context for new boi
+	USCSEditorExtensionContext* ExtensionContext = NewObject<USCSEditorExtensionContext>();
+	ExtensionContext->SubobjectEditor = SharedThis(this);
+	ExtensionContext->AddToRoot();
+
+	ButtonBox = SNew(SHorizontalBox)
+	+ SHorizontalBox::Slot()
+	.VAlign(VAlign_Center)
+	.Padding(0.0f, 0.0f, 4.0f, 0.0f)
+	.AutoWidth()
+	[
+		SNew(SComponentClassCombo)
+		.AddMetaData<FTagMetaData>(FTagMetaData(TEXT("Actor.AddComponent")))
+		.Visibility(HideComponentClassCombo.Get() ? EVisibility::Hidden : EVisibility::Visible)
+		.OnSubobjectClassSelected(this, &SSubobjectInstanceEditor::PerformComboAddClass)
+		.ToolTipText(LOCTEXT("AddComponent_Tooltip", "Adds a new component to this actor"))
+		.IsEnabled(true)
+	]
+	+ SHorizontalBox::Slot()
+	.VAlign(VAlign_Center)
+	.Padding(0,0,4,0)
+	.AutoWidth()
+	[
+		SAssignNew(ExtensionPanel, SExtensionPanel)
+		.ExtensionPanelID("SCSEditor.NextToAddComponentButton")
+		.ExtensionContext(ExtensionContext)
+	]
+	
+	// horizontal slot index #2 => reserved for BP-editor search bar (see 'ButtonBox' and 'SearchBarHorizontalSlotIndex' usage below)
+
+	+ SHorizontalBox::Slot()
+	.VAlign(VAlign_Center)
+	.Padding(0,0,4,0)
+	.AutoWidth()
+	[
+		SNew(SEditorHeaderButton)
+		.AddMetaData<FTagMetaData>(FTagMetaData(TEXT("Actor.ConvertToBlueprint")))
+		.Visibility(this, &SSubobjectInstanceEditor::GetPromoteToBlueprintButtonVisibility)
+		.OnClicked(this, &SSubobjectInstanceEditor::OnPromoteToBlueprintClicked)
+		.Icon(FAppStyle::Get().GetBrush("Icons.Blueprints"))
+		.ToolTip(IDocumentation::Get()->CreateToolTip(
+			LOCTEXT("PromoteToBluerprintTooltip", "Converts this actor into a reusable Blueprint Class that can have script behavior" ),
+			nullptr,
+			TEXT("Shared/LevelEditor"),
+			TEXT("ConvertToBlueprint")))
+	]
+	+ SHorizontalBox::Slot()
+	.VAlign(VAlign_Center)
+	.AutoWidth()
+	[
+		SNew(SEditorHeaderButton)
+		.AddMetaData<FTagMetaData>(FTagMetaData(TEXT("Actor.EditBlueprint")))
+		.Visibility(this, &SSubobjectInstanceEditor::GetEditBlueprintButtonVisibility)
+		.ToolTipText(LOCTEXT("EditActorBlueprint_Tooltip", "Edit the Blueprint for this Actor"))
+		.Icon(FAppStyle::Get().GetBrush("Icons.Blueprints"))
+		.MenuContent()
+		[
+			EditBlueprintMenuBuilder.MakeWidget()
+		]
+	];
+
+	Contents = SNew(SVerticalBox)
+	+ SVerticalBox::Slot()
+	.AutoHeight()
+	.VAlign(VAlign_Top)
+	.Padding(4.f, 0, 4.f, 4.f)
+	[
+		SAssignNew(HeaderBox, SVerticalBox)
+	]
+
+	+ SVerticalBox::Slot()
+	[
+		SNew(SBorder)
+		.BorderImage(FAppStyle::Get().GetBrush("SCSEditor.Background"))
+		.Padding(4.f)
+		.AddMetaData<FTagMetaData>(FTagMetaData(TEXT("ComponentsPanel")))
+		.Visibility(this, &SSubobjectInstanceEditor::GetComponentsTreeVisibility)
+		[
+			TreeWidget.ToSharedRef()
+		]
+	];
+
+	// Only insert the buttons and search bar in the Blueprints version
+	if (bInlineSearchBarWithButtons)
+	{
+		ButtonBox->AddSlot()
+			.FillWidth(1.0f)
+			.VAlign(VAlign_Center)
+			.Padding(3.0f, 3.0f)
+		[
+			SearchBar.ToSharedRef()
+		];
+
+		HeaderBox->AddSlot()
+			.VAlign(VAlign_Center)
+			.AutoHeight()
+		[
+			ButtonBox.ToSharedRef()
+		];
+	}
+
+	this->ChildSlot
+	[
+		Contents.ToSharedRef()
+	];
+
+	// Populate the tree with subobject data
+	UpdateTree();
+}
+END_SLATE_FUNCTION_BUILD_OPTIMIZATION
+
+void SSubobjectInstanceEditor::OnDeleteNodes()
+{
+	const FScopedTransaction Transaction( LOCTEXT("RemoveComponents", "Remove Components") );
+
+	// Invalidate any active component in the visualizer
+	GUnrealEd->ComponentVisManager.ClearActiveComponentVis();
+
+	// Gather the handles of the components that we want to delete
+	TArray<FSubobjectDataHandle> HandlesToDelete;
+	TArray<FSubobjectEditorTreeNodePtrType> SelectedNodes = GetSelectedNodes();
+	for(const FSubobjectEditorTreeNodePtrType& Node : SelectedNodes)
+	{
+		check(Node->IsValid());
+		const FSubobjectData* Data = Node->GetDataSource();
+		if(Data && Data->IsComponent())
+		{
+			HandlesToDelete.Add(Data->GetHandle());
+		}
+	}
+
+	if(USubobjectDataSubsystem* System = USubobjectDataSubsystem::Get())
+	{
+		FSubobjectDataHandle HandleToSelect;
+		int32 NumDeleted = System->DeleteSubobjects(RootNodes[0]->GetDataHandle(), HandlesToDelete, HandleToSelect);
+
+		if(NumDeleted > 0)
+		{
+			if(HandleToSelect.IsValid())
+			{
+				FSubobjectEditorTreeNodePtrType NodeToSelect = FindSlateNodeForHandle(HandleToSelect);
+				if(NodeToSelect.IsValid())
+				{
+					TreeWidget->SetSelection(NodeToSelect);
+				}
+			}
+			
+			UpdateTree();
+
+			// Do this AFTER marking the Blueprint as modified
+			UpdateSelectionFromNodes(TreeWidget->GetSelectedItems());
+		}
+	}
+}
+
+void SSubobjectInstanceEditor::CopySelectedNodes()
+{
+	TArray<FSubobjectDataHandle> SelectedHandles = GetSelectedHandles();
+	if(USubobjectDataSubsystem* System = USubobjectDataSubsystem::Get())
+	{
+		return System->CopySubobjects(SelectedHandles, /* BpContext = */ nullptr);
+	}
+}
+
+void SSubobjectInstanceEditor::OnDuplicateComponent()
+{
+	TArray<FSubobjectDataHandle> SelectedNodes = GetSelectedHandles();
+	if(SelectedNodes.Num() > 0)
+	{
+		// Force the text box being edited (if any) to commit its text. The duplicate operation may trigger a regeneration of the tree view,
+		// releasing all row widgets. If one row was in edit mode (rename/rename on create), it was released before losing the focus and
+		// this would prevent the completion of the 'rename' or 'create + give initial name' transaction (occurring on focus lost).
+		FSlateApplication::Get().ClearKeyboardFocus();
+		
+		const FScopedTransaction Transaction(SelectedNodes.Num() > 1 ? LOCTEXT("DuplicateComponents", "Duplicate Components") : LOCTEXT("DuplicateComponent", "Duplicate Component"));
+		if(USubobjectDataSubsystem* System = USubobjectDataSubsystem::Get())
+		{
+			System->DuplicateSubobjects(GetObjectContextHandle(), SelectedNodes, /* BpContext = */ nullptr);
+			UpdateTree();
+		}
+	}
+}
+
+void SSubobjectInstanceEditor::PasteNodes()
+{
+	if(USubobjectDataSubsystem* System = USubobjectDataSubsystem::Get())
+	{
+		TArray<FSubobjectDataHandle> OutHandles;
+		System->PasteSubobjects(GetObjectContextHandle(), GetSelectedHandles(), nullptr, OutHandles);
+
+		if(OutHandles.Num() > 0)
+		{
+			// We only want the pasted node(s) to be selected
+			TreeWidget->ClearSelection();
+			UpdateTree();
+
+			for(const FSubobjectDataHandle& Handle : OutHandles)
+			{
+				if(FSubobjectEditorTreeNodePtrType SlateNode = FindSlateNodeForHandle(Handle))
+				{
+					TreeWidget->SetItemSelection(SlateNode, true);	
+				}
+			}
+		}
+	}
+}
+
+void SSubobjectInstanceEditor::OnAttachToDropAction(FSubobjectEditorTreeNodePtrType DroppedOn, const TArray<FSubobjectEditorTreeNodePtrType>& DroppedNodePtrs)
+{
+	// Ask the subsystem to attach the dropped nodes onto the dropped on node
+	check(DroppedOn.IsValid());
+	check(DroppedNodePtrs.Num() > 0);
+	
+	USubobjectDataSubsystem* System = USubobjectDataSubsystem::Get();
+	check(System);
+	
+	const FScopedTransaction TransactionContext(DroppedNodePtrs.Num() > 1 ? LOCTEXT("AttachComponents", "Attach Components") : LOCTEXT("AttachComponent", "Attach Component"));
+
+	TArray<FSubobjectDataHandle> HandlesToMove;
+	Utils::PopulateHandlesArray(DroppedNodePtrs, HandlesToMove);
+
+	FReparentSubobjectParams Params;
+	Params.NewParentHandle = DroppedOn->GetDataHandle();
+	System->ReparentSubobjects(Params, HandlesToMove);
+
+	check(TreeWidget.IsValid());
+	TreeWidget->SetItemExpansion(DroppedOn, true);
+
+	PostDragDropAction(true);
+}
+
+void SSubobjectInstanceEditor::OnDetachFromDropAction(const TArray<FSubobjectEditorTreeNodePtrType>& DroppedNodePtrs)
+{
+	check(DroppedNodePtrs.Num() > 0);
+
+	const FScopedTransaction TransactionContext(DroppedNodePtrs.Num() > 1 ? LOCTEXT("DetachComponents", "Detach Components") : LOCTEXT("DetachComponent", "Detach Component"));
+
+	TArray<FSubobjectDataHandle> HandlesToMove;
+	Utils::PopulateHandlesArray(DroppedNodePtrs, HandlesToMove);
+
+	// Attach the dropped node to the current scene root node
+	FSubobjectEditorTreeNodePtrType SceneRootNodePtr = GetSceneRootNode();
+	check(SceneRootNodePtr.IsValid());
+	
+	// Ask the subsystem to reparent this object to the scene root
+	if(USubobjectDataSubsystem* System = USubobjectDataSubsystem::Get())
+	{
+		FReparentSubobjectParams Params;
+		Params.NewParentHandle = SceneRootNodePtr->GetDataHandle();
+		System->ReparentSubobjects(Params, HandlesToMove);
+	}
+
+	PostDragDropAction(false);
+}
+
+void SSubobjectInstanceEditor::OnMakeNewRootDropAction(FSubobjectEditorTreeNodePtrType DroppedNodePtr)
+{
+	// Get the current scene root node
+	FSubobjectEditorTreeNodePtrType SceneRootNodePtr = GetSceneRootNode();
+
+	// We cannot handle the drop action if any of these conditions fail on entry.
+	//if (!ensure(SceneRootNodePtr.IsValid()) || !ensure(DroppedNodePtr.IsValid()) || !ensure(DroppedNodePtr == SceneRootNodePtr))
+	//{
+	//	return;
+	//}
+	
+	// Create a transaction record
+	const FScopedTransaction TransactionContext(LOCTEXT("MakeNewSceneRoot", "Make New Scene Root"));
+
+	USubobjectDataSubsystem* Subsystem = USubobjectDataSubsystem::Get();
+	const bool bSuccess = Subsystem->MakeNewSceneRoot(
+        GetObjectContextHandle(),
+         DroppedNodePtr->GetDataHandle(),
+         nullptr);
+	
+	PostDragDropAction(true);
+}
+
+void SSubobjectInstanceEditor::PostDragDropAction(bool bRegenerateTreeNodes)
+{
+	GUnrealEd->ComponentVisManager.ClearActiveComponentVis();
+
+	UpdateTree(bRegenerateTreeNodes);
+
+	if(AActor* ActorInstance = Cast<AActor>(GetObjectContext()))
+	{
+		ActorInstance->RerunConstructionScripts();
+	}
+}
+
+TSharedPtr<SWidget> SSubobjectInstanceEditor::BuildSceneRootDropActionMenu(FSubobjectEditorTreeNodePtrType DroppedOntoNodePtr, FSubobjectEditorTreeNodePtrType DroppedNodePtr)
+{
+	FMenuBuilder MenuBuilder(true, CommandList);
+
+	const FSubobjectData* DroppedNodeData = DroppedNodePtr->GetDataSource();
+	const FSubobjectData* DroppedOntoNodeData = DroppedOntoNodePtr->GetDataSource();
+
+	check(DroppedNodeData);
+	
+	MenuBuilder.BeginSection("SceneRootNodeDropActions", LOCTEXT("SceneRootNodeDropActionContextMenu", "Drop Actions"));
+	{
+		const FText DroppedVariableNameText = FText::FromName( DroppedNodeData->GetVariableName() );
+		
+		const FText NodeVariableNameText = FText::FromName( DroppedOntoNodeData->GetVariableName() );
+
+		bool bDroppedInSameBlueprint = true;
+
+		MenuBuilder.AddMenuEntry(
+			LOCTEXT("DropActionLabel_AttachToRootNode", "Attach"),
+			bDroppedInSameBlueprint 
+			? FText::Format( LOCTEXT("DropActionToolTip_AttachToRootNode", "Attach {0} to {1}."), DroppedVariableNameText, NodeVariableNameText )
+			: FText::Format( LOCTEXT("DropActionToolTip_AttachToRootNodeFromCopy", "Copy {0} to a new variable and attach it to {1}."), DroppedVariableNameText, NodeVariableNameText ),
+			FSlateIcon(),
+			FUIAction(
+				FExecuteAction::CreateSP(this, &SSubobjectEditor::OnAttachToDropAction, DroppedOntoNodePtr, DroppedNodePtr),
+				FCanExecuteAction()));
+
+		const bool bIsDefaultSceneRoot = DroppedNodeData->IsDefaultSceneRoot();
+
+		FText NewRootNodeText = bIsDefaultSceneRoot
+			? FText::Format(LOCTEXT("DropActionToolTip_MakeNewRootNodeAndDelete", "Make {0} the new root. The default root will be deleted."), DroppedVariableNameText)
+			: FText::Format(LOCTEXT("DropActionToolTip_MakeNewRootNode", "Make {0} the new root."), DroppedVariableNameText);
+
+		FText NewRootNodeFromCopyText = bIsDefaultSceneRoot
+			? FText::Format(LOCTEXT("DropActionToolTip_MakeNewRootNodeFromCopyAndDelete", "Copy {0} to a new variable and make it the new root. The default root will be deleted."), DroppedVariableNameText)
+			: FText::Format(LOCTEXT("DropActionToolTip_MakeNewRootNodeFromCopy", "Copy {0} to a new variable and make it the new root."), DroppedVariableNameText);
+
+		MenuBuilder.AddMenuEntry(
+			LOCTEXT("DropActionLabel_MakeNewRootNode", "Make New Root"),
+			bDroppedInSameBlueprint ? NewRootNodeText : NewRootNodeFromCopyText,
+			FSlateIcon(),
+			FUIAction(
+				FExecuteAction::CreateSP(this, &SSubobjectEditor::OnMakeNewRootDropAction, DroppedNodePtr),
+				FCanExecuteAction()));
+	}
+	MenuBuilder.EndSection();
+
+	return MenuBuilder.MakeWidget();
+}
+
+FSubobjectDataHandle SSubobjectInstanceEditor::AddNewSubobject(const FSubobjectDataHandle& ParentHandle, UClass* NewClass, UObject* AssetOverride, FText& OutFailReason, TUniquePtr<FScopedTransaction> InOngoingTransaction)
+{
+	FAddNewSubobjectParams Params;
+	Params.ParentHandle = ParentHandle;
+	Params.NewClass = NewClass;
+	Params.AssetOverride = AssetOverride;
+	// This is an instance, so the blueprint context is null!
+	Params.BlueprintContext = nullptr;
+	USubobjectDataSubsystem* System = USubobjectDataSubsystem::Get();
+	check(System);
+	
+	return System->AddNewSubobject(Params,OutFailReason);	
+}
+
+void SSubobjectInstanceEditor::PopulateContextMenuImpl(UToolMenu* InMenu, TArray<FSubobjectEditorTreeNodePtrType>& InSelectedItems, bool bIsChildActorSubtreeNodeSelected)
+{
+	TArray<UActorComponent*> SelectedComponents;
+	for(const FSubobjectEditorTreeNodePtrType& SelectedNodePtr : InSelectedItems)
+	{
+		check(SelectedNodePtr->IsValid());
+		// Get the component template associated with the selected node
+		const UActorComponent* ComponentTemplate = SelectedNodePtr->GetComponentTemplate();
+		if (ComponentTemplate)
+		{
+			// #TODO_BH Remove this const cast
+			SelectedComponents.Add(const_cast<UActorComponent*>(ComponentTemplate));
+		}
+	}
+	
+	// Common menu options added for all component types
+    FComponentEditorUtils::FillComponentContextMenuOptions(InMenu, SelectedComponents);
+}
+
+FSlateColor SSubobjectInstanceEditor::GetColorTintForIcon(FSubobjectEditorTreeNodePtrType Node) const
+{
+	static const FLinearColor IntroducedHereColor(FLinearColor::White);
+	// A blue-ish tint
+	static const FLinearColor InheritedBlueprintComponentColor(0.08f, 0.35f, 0.6f);
+	
+	static const FLinearColor InstancedInheritedBlueprintComponentColor(0.08f, 0.35f, 0.6f);
+	// A green-ish tint
+	static const FLinearColor InheritedNativeComponentColor(0.7f, 0.9f, 0.7f);
+	
+	const FSubobjectData* Data = Node ? Node->GetDataSource() : nullptr;
+	if(!Data)
+	{
+		return IntroducedHereColor;
+	}
+
+	if (Data->IsInheritedComponent())
+	{
+		// Native C++ components will be tinted green
+		if (Data->IsNativeComponent())
+		{
+			return InheritedNativeComponentColor;
+		}
+		else if (Data->IsInstancedComponent())
+		{
+			return InstancedInheritedBlueprintComponentColor;
+		}
+		else
+		{
+			return InheritedBlueprintComponentColor;
+		}
+	}
+	// If its not an actor, instanced, or native inherited component then it must be inherited from a BP
+	else if(!Data->IsActor() && !Data->IsInstancedComponent())
+	{
+		return InheritedBlueprintComponentColor;
+	}
+	else if(Data->IsActor())
+	{
+		return FSlateColor::UseForeground();
+	}
+	
+	return IntroducedHereColor;
+}
+
+#undef LOCTEXT_NAMESPACE
diff --git a/Engine/Source/Editor/SubobjectEditor/Public/SSubobjectBlueprintEditor.h b/Engine/Source/Editor/SubobjectEditor/Public/SSubobjectBlueprintEditor.h
new file mode 100644
--- /dev/null
+++ b/Engine/Source/Editor/SubobjectEditor/Public/SSubobjectBlueprintEditor.h
@@ -0,0 +1,104 @@
+// Copyright Epic Games, Inc. All Rights Reserved.
+
+#pragma once
+
+#include "SSubobjectEditor.h"
+
+/**
+* This is the editor for subobjects within the blueprint editor that
+*/
+class SUBOBJECTEDITOR_API SSubobjectBlueprintEditor final : public SSubobjectEditor
+{
+private:	
+	SLATE_BEGIN_ARGS(SSubobjectBlueprintEditor)
+        : _ObjectContext(nullptr)
+		, _PreviewActor(nullptr)
+        , _AllowEditing(true)
+        , _OnSelectionUpdated()
+		{}
+
+		SLATE_ATTRIBUTE(UObject*, ObjectContext)
+		SLATE_ATTRIBUTE(AActor*, PreviewActor)
+	    SLATE_ATTRIBUTE(bool, AllowEditing)
+	    SLATE_EVENT(FOnSelectionUpdated, OnSelectionUpdated)
+	    SLATE_EVENT(FOnItemDoubleClicked, OnItemDoubleClicked)
+	    
+	SLATE_END_ARGS()
+	
+	void Construct(const FArguments& InArgs);
+
+protected:
+
+	/** Attribute that provides access to a "preview" Actor context (may not be same as the Actor context that's being edited. */
+	TAttribute<AActor*> PreviewActor;
+	
+	// SSubobjectEditor interface
+	virtual void OnDeleteNodes() override;
+	virtual void CopySelectedNodes() override;
+	virtual void OnDuplicateComponent() override;
+	virtual bool CanPasteNodes() const override;
+	virtual void PasteNodes() override;
+	
+	virtual void OnAttachToDropAction(FSubobjectEditorTreeNodePtrType DroppedOn, const TArray<FSubobjectEditorTreeNodePtrType>& DroppedNodePtrs);
+    virtual void OnDetachFromDropAction(const TArray<FSubobjectEditorTreeNodePtrType>& DroppedNodePtrs);
+    virtual void OnMakeNewRootDropAction(FSubobjectEditorTreeNodePtrType DroppedNodePtr);
+    virtual void PostDragDropAction(bool bRegenerateTreeNodes);
+
+    /** Builds a context menu pop up for dropping a child node onto the scene root node */
+    virtual TSharedPtr<SWidget> BuildSceneRootDropActionMenu(FSubobjectEditorTreeNodePtrType DroppedOntoNodePtr, FSubobjectEditorTreeNodePtrType DroppedNodePtr);
+	virtual bool ClearSelectionOnClick() const override { return true; }
+	virtual bool ShowInlineSearchWithButtons() const override { return true; }
+	virtual FSubobjectDataHandle AddNewSubobject(const FSubobjectDataHandle& ParentHandle, UClass* NewClass, UObject* AssetOverride, FText& OutFailReason, TUniquePtr<FScopedTransaction> InOngoingTransaction) override;
+	virtual void PopulateContextMenuImpl(UToolMenu* InMenu, TArray<FSubobjectEditorTreeNodePtrType>& InSelectedItems, bool bIsChildActorSubtreeNodeSelected) override;
+	virtual FSubobjectEditorTreeNodePtrType GetSceneRootNode() const override;
+public:
+	//virtual FSubobjectEditorTreeNodePtrType FindSlateNodeForObject(const UObject* InObject, bool bIncludeAttachmentComponents = true) const override;
+	// End of SSubobjectEditor
+
+public:
+
+	/**
+	* Fills out an events section in ui.
+	* @param Menu								the menu to add the events section into
+	* @param Blueprint							the active blueprint context being edited
+	* @param SelectedClass						the common component class to build the events list from
+	* @param CanExecuteActionDelegate			the delegate to query whether or not to execute the UI action
+	* @param GetSelectedObjectsDelegate		the delegate to fill the currently select variables / components
+	*/
+	static void BuildMenuEventsSection(FMenuBuilder& Menu, UBlueprint* Blueprint, UClass* SelectedClass, FCanExecuteAction CanExecuteActionDelegate, FGetSelectedObjectsDelegate GetSelectedObjectsDelegate);
+
+protected:
+
+	/**
+	* Function to create events for the current selection
+	* @param Blueprint						the active blueprint context
+	* @param EventName						the event to add
+	* @param GetSelectedObjectsDelegate	the delegate to gather information about current selection
+	* @param NodeIndex						an index to a specified node to add event for or < 0 for all selected nodes.
+	*/
+	static void CreateEventsForSelection(UBlueprint* Blueprint, FName EventName, FGetSelectedObjectsDelegate GetSelectedObjectsDelegate);
+
+	/**
+	* Function to construct an event for a node
+	* @param Blueprint						the nodes blueprint
+	* @param EventName						the event to add
+	* @param EventData						the event data structure describing the node
+	*/
+	static void ConstructEvent(UBlueprint* Blueprint, const FName EventName, const FComponentEventConstructionData EventData);
+
+	/**
+	* Function to view an event for a node
+	* @param Blueprint						the nodes blueprint
+	* @param EventName						the event to view
+	* @param EventData						the event data structure describing the node
+	*/
+	static void ViewEvent(UBlueprint* Blueprint, const FName EventName, const FComponentEventConstructionData EventData);
+
+	/** Get the current preview actor for this blueprint editor. */
+	AActor* GetActorPreview() const { return PreviewActor.Get(nullptr); }
+
+public:
+	
+	virtual FSlateColor GetColorTintForIcon(FSubobjectEditorTreeNodePtrType Node) const override;
+	
+};
\ No newline at end of file
diff --git a/Engine/Source/Editor/SubobjectEditor/Public/SSubobjectEditor.h b/Engine/Source/Editor/SubobjectEditor/Public/SSubobjectEditor.h
new file mode 100644
--- /dev/null
+++ b/Engine/Source/Editor/SubobjectEditor/Public/SSubobjectEditor.h
@@ -0,0 +1,728 @@
+// Copyright Epic Games, Inc. All Rights Reserved.
+
+#pragma once
+
+#include "CoreMinimal.h"
+#include "Components/ActorComponent.h"
+#include "Components/SceneComponent.h"
+#include "SubobjectDataSubsystem.h"
+
+#include "Widgets/SCompoundWidget.h"
+#include "Widgets/Views/STableRow.h"
+#include "Widgets/Views/STreeView.h"
+#include "Framework/Commands/UICommandList.h"
+#include "SComponentClassCombo.h"
+
+#include "BlueprintEditor.h"		// FComponentEventConstructionData
+
+#include "ScopedTransaction.h"
+
+class SSubobjectEditor;
+class SToolTip;
+class SExtensionPanel;
+class AActor;
+class FSubobjectEditorTreeNode;
+class SInlineEditableTextBlock;
+class ISCSEditorUICustomization;
+class UBlueprint;
+class UToolMenu;
+class UPrimitiveComponent;
+
+// Tree node pointer types
+using FSubobjectEditorTreeNodePtrType = TSharedPtr<class FSubobjectEditorTreeNode>;
+
+/////////////////////////////////////////////////////
+// FSubobjectEditorTreeNode
+
+/**
+ * Wrapper struct that represents access an singe subobject that is used 
+ * by slate as a layout for columns/rows should look like
+ */
+class SUBOBJECTEDITOR_API FSubobjectEditorTreeNode : public TSharedFromThis<FSubobjectEditorTreeNode>
+{
+public:
+
+	/** Delegate for when the context menu requests a rename */
+	DECLARE_DELEGATE(FOnRenameRequested);
+	
+	explicit FSubobjectEditorTreeNode(const FSubobjectDataHandle& DataSource, bool InbIsSeperator = false);
+
+	~FSubobjectEditorTreeNode() = default;
+
+	bool IsSeperator() const { return bIsSeperator; } 
+
+	// If this subobject is an actor component, then return its template. Otherwise it will
+	// return null
+	const UActorComponent* GetComponentTemplate() const;
+
+	const UObject* GetObject(bool bEvenIfPendingKill = false) const;
+
+	// Get a pointer to the subobject data that this slate node is representing
+	FSubobjectData* GetDataSource() const;
+
+	// Get the subobject handle that this slate node is representing
+	FSubobjectDataHandle GetDataHandle() const { return DataHandle; }
+
+	/** Get the slate parent of this single tree node */
+	FSubobjectEditorTreeNodePtrType GetParent() const { return ParentNodePtr; }
+
+	/** Get the display name of this single tree node */
+	FString GetDisplayString() const;
+
+	bool IsComponentNode() const;
+
+	// Returns true if this is the root actor node of the tree
+	bool IsRootActorNode() const;
+
+	bool CanReparent() const;
+
+	FName GetVariableName() const;
+
+	// Returns true if this subobject data handle is valid
+	bool IsValid() const { return bIsSeperator || DataHandle.IsValid(); }
+
+	bool IsChildSubtreeNode() const;
+
+	bool IsNativeComponent() const;
+	
+	// Returns true if this node is attached to the given slate node
+	bool IsAttachedTo(FSubobjectEditorTreeNodePtrType InNodePtr) const;
+
+	/**
+	* @return Whether or not this node is a direct child of the given node.
+	*/
+	bool IsDirectlyAttachedTo(FSubobjectEditorTreeNodePtrType InNodePtr) const;
+	
+	bool CanDelete() const;
+
+	bool CanRename() const;
+	
+	/**
+	 * @return The set of nodes which are parented to this node (read-only).
+	 */
+	const TArray<FSubobjectEditorTreeNodePtrType>& GetChildren() const { return Children; }
+
+	// Add the given slate node to our child array and set it's parent to us
+	// This had no effect on the actual structure of the subobjects this node represents
+	// it is purely visual
+	void AddChild(FSubobjectEditorTreeNodePtrType AttachToPtr);
+
+	// Remove the given slate node from our children array
+	void RemoveChild(FSubobjectEditorTreeNodePtrType InChildNodePtr);
+
+	// Attempts to find the given subobject handle in the slate children on this node. Nullptr if none are found.
+	FSubobjectEditorTreeNodePtrType FindChild(const FSubobjectDataHandle& InHandle);
+
+	/** Query that determines if this item should be filtered out or not */
+	bool IsFlaggedForFiltration() const;
+
+	/** Sets this item's filtration state. Use bUpdateParent to make sure the parent's EFilteredState::ChildMatches flag is properly updated based off the new state */
+	void SetCachedFilterState(bool bMatchesFilter, bool bUpdateParent);
+
+	/** Used to update the EFilteredState::ChildMatches flag for parent nodes, when this item's filtration state has changed */
+	void ApplyFilteredStateToParent();
+
+	/** Updates the EFilteredState::ChildMatches flag, based off of children's current state */
+	void RefreshCachedChildFilterState(bool bUpdateParent);
+
+	/** Refreshes this item's filtration state. Set bRecursive to 'true' to refresh any child nodes as well */
+	bool RefreshFilteredState(const UClass* InFilterType, const TArray<FString>& InFilterTerms, bool bRecursive);
+
+	/** Returns whether the node will match the given type (for filtering) */
+	bool MatchesFilterType(const UClass* InFilterType) const;
+
+	void SetOngoingCreateTransaction(TUniquePtr<FScopedTransaction> InTransaction);
+	void CloseOngoingCreateTransaction();
+	void GetOngoingCreateTransaction(TUniquePtr<FScopedTransaction>& OutPtr) { OutPtr = MoveTemp(OngoingCreateTransaction); }
+	bool HasOngoingTransaction() const { return OngoingCreateTransaction.IsValid(); } 
+
+	/** Sets up the delegate for a rename operation */
+	void SetRenameRequestedDelegate(FOnRenameRequested InRenameRequested) { RenameRequestedDelegate = InRenameRequested; }
+	FOnRenameRequested GetRenameRequestedDelegate() { return RenameRequestedDelegate; }
+
+protected: 
+
+	/** Pointer to the parent of this subobject */
+	FSubobjectEditorTreeNodePtrType ParentNodePtr;
+	
+	// Scope the creation of a node which ends when the initial 'name' is given/accepted by the user, which can be several frames after the node was actually created.
+	TUniquePtr<FScopedTransaction> OngoingCreateTransaction;
+
+	/** Handles rename requests */
+	FOnRenameRequested RenameRequestedDelegate;
+	
+	/**
+	* Any children that this subobject has in the hierarchy, used to
+	* collapsed things within in the tree
+	*/
+	TArray<FSubobjectEditorTreeNodePtrType> Children;
+
+	/** The data source of this subobject */
+	FSubobjectDataHandle DataHandle;
+
+	enum EFilteredState
+	{
+		FilteredOut		= 0x00,
+		MatchesFilter	= (1 << 0),
+		ChildMatches	= (1 << 1),
+
+		FilteredInMask = (MatchesFilter | ChildMatches),
+		Unknown = 0xFC // ~FilteredInMask
+	};
+	uint8 FilterFlags;
+
+	/** A flag that indicates that this is a separator slate node and has no valid data */
+	uint8 bIsSeperator : 1;
+};
+
+/////////////////////////////////////////////////////
+// SSubobjectEditorDragDropTree
+
+/** Implements the specific node type and add drag/drop functionality */
+class SSubobjectEditorDragDropTree : public STreeView<FSubobjectEditorTreeNodePtrType>
+{
+public:
+	SLATE_BEGIN_ARGS(SSubobjectEditorDragDropTree)
+		: _SubobjectEditor(nullptr)
+		, _OnGenerateRow()
+		, _OnGetChildren()
+		, _OnSetExpansionRecursive()
+		, _TreeItemsSource(static_cast<TArray<FSubobjectEditorTreeNodePtrType>*>(nullptr))
+		, _ItemHeight(16)
+		, _OnContextMenuOpening()
+		, _OnMouseButtonDoubleClick()
+		, _OnSelectionChanged()
+		, _SelectionMode(ESelectionMode::Multi)
+		, _ClearSelectionOnClick(true)
+		, _ExternalScrollbar()
+		, _OnTableViewBadState()
+	{
+		_Clipping = EWidgetClipping::ClipToBounds;
+	}
+
+		SLATE_ARGUMENT(SSubobjectEditor*, SubobjectEditor)
+
+		SLATE_EVENT(FOnGenerateRow, OnGenerateRow)
+
+		SLATE_EVENT(FOnItemScrolledIntoView, OnItemScrolledIntoView)
+
+		SLATE_EVENT(FOnGetChildren, OnGetChildren)
+
+		SLATE_EVENT(FOnSetExpansionRecursive, OnSetExpansionRecursive)
+
+		SLATE_ARGUMENT(TArray<FSubobjectEditorTreeNodePtrType>*, TreeItemsSource)
+
+		SLATE_ATTRIBUTE(float, ItemHeight)
+
+		SLATE_EVENT(FOnContextMenuOpening, OnContextMenuOpening)
+
+		SLATE_EVENT(FOnMouseButtonDoubleClick, OnMouseButtonDoubleClick)
+
+		SLATE_EVENT(FOnSelectionChanged, OnSelectionChanged)
+
+		SLATE_ATTRIBUTE(ESelectionMode::Type, SelectionMode)
+
+		SLATE_ARGUMENT(TSharedPtr<SHeaderRow>, HeaderRow)
+
+		SLATE_ARGUMENT(bool, ClearSelectionOnClick)
+
+		SLATE_ARGUMENT(TSharedPtr<SScrollBar>, ExternalScrollbar)
+
+		SLATE_EVENT(FOnTableViewBadState, OnTableViewBadState)
+
+	SLATE_END_ARGS()
+
+	/** Object construction - mostly defers to the base STreeView */
+	void Construct(const FArguments& InArgs);
+
+	// SWidget interface
+	virtual FReply OnDragOver(const FGeometry& MyGeometry, const FDragDropEvent& DragDropEvent) override;
+	virtual FReply OnDrop(const FGeometry& MyGeometry, const FDragDropEvent& DragDropEvent) override;
+	// End SWidget interface
+
+private:
+
+	SSubobjectEditor* SubobjectEditor;
+};
+
+/////////////////////////////////////////////////////
+// SSubobject_RowWidget
+
+/**
+* A row widget that represents a single subobject within the tree. 
+*/
+class SSubobject_RowWidget : public SMultiColumnTableRow<FSubobjectEditorTreeNodePtrType>
+{
+public:
+	SLATE_BEGIN_ARGS(SSubobject_RowWidget) {}
+	SLATE_END_ARGS()
+	
+	void Construct(const FArguments& InArgs, TWeakPtr<SSubobjectEditor> InEditor, FSubobjectEditorTreeNodePtrType InNodePtr, TSharedPtr<STableViewBase> InOwnerTableView);
+
+	// SMultiColumnTableRow<T> interface
+	virtual TSharedRef<SWidget> GenerateWidgetForColumn(const FName& ColumnName) override;
+
+protected:
+	virtual ESelectionMode::Type GetSelectionMode() const override;
+	// End of SMultiColumnTableRow<T>
+
+	/** Drag-drop handlers */
+	void HandleOnDragEnter(const FDragDropEvent& DragDropEvent);
+	void HandleOnDragLeave(const FDragDropEvent& DragDropEvent);
+	FReply HandleOnDragDetected(const FGeometry& MyGeometry, const FPointerEvent& MouseEvent);
+	TOptional<EItemDropZone> HandleOnCanAcceptDrop(const FDragDropEvent& DragDropEvent, EItemDropZone DropZone, FSubobjectEditorTreeNodePtrType TargetItem);
+	FReply HandleOnAcceptDrop(const FDragDropEvent& DragDropEvent, EItemDropZone DropZone, FSubobjectEditorTreeNodePtrType TargetItem);
+
+public:	
+	/** Get the asset name of this subobject from the asset brokerage */
+	FText GetAssetName() const;
+
+	/** Get the asset path of this subobject from the asset brokerage */
+	FText GetAssetPath() const;
+
+	FText GetNameLabel() const;
+
+	/** Check if this asset is visible from the asset brokerage */
+	EVisibility GetAssetVisibility() const;
+
+	virtual const FSlateBrush* GetIconBrush() const;
+
+	FSlateColor GetColorTintForIcon() const;
+
+	FSubobjectEditorTreeNodePtrType GetSubobjectPtr() const { return SubobjectPtr; }
+	
+private:
+
+	static void AddToToolTipInfoBox(const TSharedRef<SVerticalBox>& InfoBox, const FText& Key, TSharedRef<SWidget> ValueIcon, const TAttribute<FText>& Value, bool bImportant);
+
+	/** Creates a tooltip for this row */
+	TSharedRef<SToolTip> CreateToolTipWidget() const;
+
+	/** Create the tooltip for component subobjects */
+	TSharedRef<SToolTip> CreateComponentTooltipWidget(const FSubobjectEditorTreeNodePtrType& InNode) const;
+	
+	/** Create the tooltip for actor subobjects */
+	TSharedRef<SToolTip> CreateActorTooltipWidget(const FSubobjectEditorTreeNodePtrType& InNode) const;
+	
+	FText GetTooltipText() const;
+	
+	FText GetActorClassNameText() const;
+	FText GetActorSuperClassNameText() const;
+	FText GetActorMobilityText() const;
+	FText GetObjectContextText() const;
+
+	FString GetDocumentationLink() const;
+	FString GetDocumentationExcerptName() const;
+
+	/**
+	 * Retrieves tooltip text describing the specified component's mobility.
+	 *
+	 * @returns An FText object containing a description of the component's mobility
+	 */
+	FText GetMobilityToolTipText() const;
+	
+	/**
+	 * Retrieves an image brush signifying the specified component's mobility (could sometimes be NULL).
+	 *
+	 * @returns A pointer to the FSlateBrush to use (NULL for Static and Non-SceneComponents)
+	 */
+	FSlateBrush const* GetMobilityIconImage() const;
+
+	/**
+	* Retrieves tooltip text describing where the component was first introduced (for inherited components).
+	* 
+	* @returns An FText object containing a description of when the component was first introduced
+	*/
+	FText GetIntroducedInToolTipText() const;
+	
+	/**
+	* Retrieves tooltip text describing how the component was introduced
+	* 
+	* @returns An FText object containing a description of when the component was first introduced
+	*/
+	FText GetComponentAddSourceToolTipText() const;
+
+	/**
+	* Retrieves a tooltip text describing if the component is marked Editor only or not
+	*
+	* @returns An FText object containing a description of if the component is marked Editor only or not
+	*/
+	FText GetComponentEditorOnlyTooltipText() const;
+	
+	/**
+	* Retrieves tooltip text for the specified Native Component's underlying Name
+	*
+	* @returns An FText object containing the Component's Name
+	*/
+	FText GetNativeComponentNameToolTipText() const;
+
+	FText GetActorDisplayText() const;
+
+	/** Commits the new name of the component */
+	void OnNameTextCommit(const FText& InNewName, ETextCommit::Type InTextCommit);
+
+	/** Verifies the name of the component when changing it */
+	bool OnNameTextVerifyChanged(const FText& InNewText, FText& OutErrorMessage);
+
+	bool IsReadOnly() const;
+
+	TWeakPtr<SSubobjectEditor> SubobjectEditor;
+	TSharedPtr<SInlineEditableTextBlock> InlineWidget;
+
+	FSubobjectEditorTreeNodePtrType SubobjectPtr;
+};
+
+/////////////////////////////////////////////////////
+// SSubobjectEditor
+
+/**
+* The base class viewer for subobject editing in Slate. This displays
+* all subobjects of a given actor or blueprint. A subobject can be a native
+* component on an actor or inherited components from a native or blueprint parent.
+*/
+class SUBOBJECTEDITOR_API SSubobjectEditor : public SCompoundWidget
+{
+public:
+
+	DECLARE_DELEGATE_OneParam(FOnSelectionUpdated, const TArray<FSubobjectEditorTreeNodePtrType>&);
+	DECLARE_DELEGATE_OneParam(FOnItemDoubleClicked, const FSubobjectEditorTreeNodePtrType);
+	
+protected:
+	
+	// Do not allow public construction of this widget! 
+	SSubobjectEditor() = default;
+
+	virtual ~SSubobjectEditor() = default;
+	
+	/** Delegate to invoke on selection update. */
+	FOnSelectionUpdated OnSelectionUpdated;
+
+	/** Delegate to invoke when an item in the tree is double clicked. */
+	FOnItemDoubleClicked OnItemDoubleClicked;
+
+	/** Attribute that provides access to the Object context for which we are viewing/editing. */
+	TAttribute<UObject*> ObjectContext;
+
+	/** Attribute to indicate whether or not editing is allowed. */
+	TAttribute<bool> AllowEditing;
+
+	/** Attribute to indicate whether or not the "Add Component" button is visible. If true, new components cannot be added to the Blueprint. */
+	TAttribute<bool> HideComponentClassCombo;
+
+	/** Attribute to limit visible nodes to a particular component type when filtering the tree view. */
+	TAttribute<TSubclassOf<UActorComponent>> ComponentTypeFilter;
+
+	/////////////////////////////////////////////////////////
+	// Widget Callbacks
+protected:
+
+	TSharedRef<ITableRow> MakeTableRowWidget(FSubobjectEditorTreeNodePtrType InNodePtr, const TSharedRef<STableViewBase>& OwnerTable);
+
+	/** @return The visibility of the components tree */
+	EVisibility GetComponentsTreeVisibility() const;
+
+	/** Used by tree control - get children for a specified subobject node */
+	void OnGetChildrenForTree(FSubobjectEditorTreeNodePtrType InNodePtr, TArray<FSubobjectEditorTreeNodePtrType>& OutChildren);
+
+	/** Update any associated selection (e.g. details view) from the passed in nodes */
+	void UpdateSelectionFromNodes(const TArray<FSubobjectEditorTreeNodePtrType>& SelectedNodes);
+
+	/** Called when selection in the tree changes */
+	void OnTreeSelectionChanged(FSubobjectEditorTreeNodePtrType InSelectedNodePtr, ESelectInfo::Type SelectInfo);
+
+	/** Callback when a component item is double clicked. */
+	void HandleItemDoubleClicked(FSubobjectEditorTreeNodePtrType InItem);
+	
+	void OnFindReferences();
+
+	/** @return The visibility of the components filter box */
+	EVisibility GetComponentsFilterBoxVisibility() const;
+
+	/** Recursively updates the filtered state for each component item */
+	void OnFilterTextChanged(const FText& InFilterText);
+
+	/** Callback when a component item is scrolled into view */
+	void OnItemScrolledIntoView(FSubobjectEditorTreeNodePtrType InItem, const TSharedPtr<ITableRow>& InWidget);
+
+	/** SWidget interface */
+	virtual FReply OnKeyDown(const FGeometry& MyGeometry, const FKeyEvent& InKeyEvent) override;
+	
+public:
+	/** Select the given tree node */
+	void SelectNode(FSubobjectEditorTreeNodePtrType InNodeToSelect, bool bIsCntrlDown);
+
+	/** Select the given tree node if there is on that matches the given handle */
+	void SelectNodeFromHandle(const FSubobjectDataHandle& InHandle, bool bIsCntrlDown);
+	
+	/** Update the contents of the Subobject Tree based on the current context */
+	void UpdateTree(bool bRegenerateTreeNodes = true);
+
+	/** Forces the details panel to refresh on the same objects */
+	void RefreshSelectionDetails();
+
+	/** Clears the current selection */
+	void ClearSelection();
+
+	/** Select the root of the tree */
+	void SelectRoot();
+
+	virtual FSlateColor GetColorTintForIcon(FSubobjectEditorTreeNodePtrType Node) const = 0;
+	
+	/** Get the currently selected nodes from the tree sorted in order from parent to child */
+	TArray<FSubobjectEditorTreeNodePtrType> GetSelectedNodes() const;
+
+	/** Returns the number of currently selected nodes in the tree */
+	int32 GetNumSelectedNodes() const { return TreeWidget->GetSelectedItems().Num(); }
+
+	/** Get the currently selected handles from the tree sorted in order from parent to child */
+	TArray<FSubobjectDataHandle> GetSelectedHandles() const;
+
+	virtual FSubobjectEditorTreeNodePtrType GetSceneRootNode() const;
+	
+	/** Try to handle a drag-drop operation */
+	FReply TryHandleAssetDragDropOperation(const FDragDropEvent& DragDropEvent);
+	
+	/** Sets UI customizations of this SCSEditor. */
+	void SetUICustomization(TSharedPtr<ISCSEditorUICustomization> InUICustomization);
+
+	virtual bool IsEditingAllowed() const;
+	
+	/** Return the button widgets that can add components or create/edit blueprints */
+	TSharedPtr<SWidget> GetToolButtonsBox() const;
+	
+	TSharedPtr<FUICommandList> GetCommandList() const { return CommandList; }
+
+	TSharedPtr<SSubobjectEditorDragDropTree> GetDragDropTree() const { return TreeWidget; }
+
+	// Drag/drop operations
+	/** Handler for attaching a single node to this node */
+	void OnAttachToDropAction(FSubobjectEditorTreeNodePtrType DroppedOn, FSubobjectEditorTreeNodePtrType DroppedNodePtr)
+	{
+		TArray<FSubobjectEditorTreeNodePtrType> DroppedNodePtrs;
+		DroppedNodePtrs.Add(DroppedNodePtr);
+		OnAttachToDropAction(DroppedOn, DroppedNodePtrs);
+	}
+	
+	virtual void OnAttachToDropAction(FSubobjectEditorTreeNodePtrType DroppedOn, const TArray<FSubobjectEditorTreeNodePtrType>& DroppedNodePtrs) = 0;
+	virtual void OnDetachFromDropAction(const TArray<FSubobjectEditorTreeNodePtrType>& DroppedNodePtrs) = 0;
+	virtual void OnMakeNewRootDropAction(FSubobjectEditorTreeNodePtrType DroppedNodePtr) = 0;
+	virtual void PostDragDropAction(bool bRegenerateTreeNodes) = 0;
+
+	/** Builds a context menu popup for dropping a child node onto the scene root node */
+	virtual TSharedPtr<SWidget> BuildSceneRootDropActionMenu(FSubobjectEditorTreeNodePtrType DroppedOntoNodePtr, FSubobjectEditorTreeNodePtrType DroppedNodePtr) = 0;
+	
+	virtual bool CanMakeNewRootOnDrag(UBlueprint* DraggedFromBlueprint) const { return false; }
+
+	/** Provides access to the Blueprint context that's being edited */
+	UBlueprint* GetBlueprint() const;
+
+	// Attempt to find an existing slate node that matches the given handle
+	virtual FSubobjectEditorTreeNodePtrType FindSlateNodeForObject(const UObject* InObject, bool bIncludeAttachmentComponents = true) const;
+	
+	/** Returns true if the specified component is currently selected */
+    bool IsComponentSelected(const UPrimitiveComponent* PrimComponent) const;
+
+    /** Assigns a selection override delegate to the specified component */
+    void SetSelectionOverride(UPrimitiveComponent* PrimComponent) const;
+
+	/**
+	* Fills the supplied array with the currently selected objects
+	* @param OutSelectedItems The array to fill.
+	*/
+	void GetSelectedItemsForContextMenu(TArray<FComponentEventConstructionData>& OutSelectedItems) const;
+
+	/** Return an array of the current slate node hierarchy in the tree */
+	const TArray<FSubobjectEditorTreeNodePtrType>& GetRootNodes() const { return RootNodes; };
+
+	// Attempt to find an existing slate node that matches the given handle
+	FSubobjectEditorTreeNodePtrType FindSlateNodeForHandle(const FSubobjectDataHandle& Handle, FSubobjectEditorTreeNodePtrType InStartNodePtr = FSubobjectEditorTreeNodePtrType()) const;
+	
+protected:
+
+	/** Restore the previous selection state when updating the tree */
+	virtual void RestoreSelectionState(TArray<FSubobjectEditorTreeNodePtrType>& SelectedTreeNodes);
+	
+	bool CanCutNodes() const;
+	void CutSelectedNodes();
+
+	bool CanCopyNodes() const;
+	virtual void CopySelectedNodes() = 0;
+
+	/** Pastes previously copied node(s) */
+	virtual bool CanPasteNodes() const;
+	virtual void PasteNodes() = 0;
+	
+public:
+	/** Removes existing selected component nodes from the SCS */
+	virtual bool CanDeleteNodes() const;
+	virtual void OnDeleteNodes() = 0;
+	
+protected:
+	bool CanDuplicateComponent() const;
+	virtual void OnDuplicateComponent() = 0;
+
+	/** Checks to see if renaming is allowed on the selected component */
+	bool CanRenameComponent() const;
+	void OnRenameComponent();
+	void OnRenameComponent(TUniquePtr<FScopedTransaction> InComponentCreateTransaction);
+	
+	/** Called at the end of each frame. */
+	void OnPostTick(float);
+	
+	/////////////////////////////////////////////////////////
+	// Widgets
+
+	/** Root set of tree that represents all subobjects for the current context */
+	TArray<FSubobjectEditorTreeNodePtrType> RootNodes;
+
+	/** Tree widget */
+	TSharedPtr<SSubobjectEditorDragDropTree> TreeWidget;
+
+	/** Command list for handling actions in the editor. */
+	TSharedPtr<FUICommandList> CommandList;
+
+	/** The filter box that handles filtering for the tree. */
+	TSharedPtr<SSearchBox> FilterBox;
+
+	/** The tools buttons box */
+	TSharedPtr<SHorizontalBox> ButtonBox;
+
+	/** Gate to prevent changing the selection while selection change is being broadcast. */
+	bool bUpdatingSelection;
+
+	/** Controls whether or not to allow calls to UpdateTree() */
+	bool bAllowTreeUpdates;
+
+	/** SCSEditor UI customizations */
+	TSharedPtr<ISCSEditorUICustomization> UICustomization;
+
+	/** SCSEditor UI extension */
+	TSharedPtr<SExtensionPanel> ExtensionPanel;
+
+	/** Scope the creation of a component which ends when the initial component 'name' is given/accepted by the user, which can be several frames after the component was actually created. */
+	TUniquePtr<FScopedTransaction> DeferredOngoingCreateTransaction;
+
+	/** Used to unregister from the post tick event. */
+	FDelegateHandle PostTickHandle;
+
+	/** Name of a node that has been requested to be renamed */
+	FSubobjectDataHandle DeferredRenameRequest;
+
+	/**
+	 * The handle to the current root context. If this is different in between
+	 * UpdateTree calls, then we know the context has changed and we should clean up
+	 * the subobject memory layout. 
+	 */
+	FSubobjectDataHandle CachedRootHandle;
+
+	virtual FMenuBuilder CreateMenuBuilder() const;
+
+	/** Constructs the slate drag/drop tree for this subobject editor */
+	virtual void ConstructTreeWidget();
+
+	/** Creates a list of commands */
+	virtual void CreateCommandList();
+	
+	UObject* GetObjectContext() const;
+
+	FSubobjectDataHandle GetObjectContextHandle() const;
+
+	/** Recursively visits the given node + its children and invokes the given function for each. */
+	void DepthFirstTraversal(const FSubobjectEditorTreeNodePtrType& InNodePtr, TSet<FSubobjectEditorTreeNodePtrType>& OutVisitedNodes, const TFunctionRef<void(const FSubobjectEditorTreeNodePtrType&)> InFunction) const;
+
+	/** Returns the set of expandable nodes that are currently collapsed in the UI */
+	void GetCollapsedNodes(const FSubobjectEditorTreeNodePtrType& InNodePtr, TSet<FSubobjectEditorTreeNodePtrType>& OutCollapsedNodes) const;
+
+	// Attempt to find an existing slate node that matches the given handle
+	FSubobjectEditorTreeNodePtrType FindOrCreateSlateNodeForHandle(const FSubobjectDataHandle& Handle) const;
+
+	/**
+	 * Set the expansion state of a node
+	 *
+	 * @param InNodeToChange	The node to be expanded/collapsed
+	 * @param bIsExpanded		True to expand the node, false to collapse it
+	 */
+	void SetNodeExpansionState(FSubobjectEditorTreeNodePtrType InNodeToChange, const bool bIsExpanded);
+	
+	/** Handler for recursively expanding/collapsing items */
+	void SetItemExpansionRecursive(FSubobjectEditorTreeNodePtrType Model, bool bInExpansionState);
+
+	/** Dumps out the tree view contents to the log (used to assist with debugging widget hierarchy issues) */
+	void DumpTree();
+
+	/** Registers context menu by name for later access */
+	void RegisterContextMenu();
+
+	/** Called to display context menu when right clicking on the widget */
+	TSharedPtr<SWidget> CreateContextMenu();
+
+	/** Populate context menu on the fly */
+	void PopulateContextMenu(UToolMenu* InMenu);
+
+	/** Populate the context menu with implementation specific details */
+	virtual void PopulateContextMenuImpl(UToolMenu* InMenu, TArray<FSubobjectEditorTreeNodePtrType>& InSelectedItems, bool bIsChildActorSubtreeNodeSelected) = 0;
+	
+	/** @return Type of component to filter the tree view with or nullptr if there's no filter. */
+	TSubclassOf<UActorComponent> GetComponentTypeFilterToApply() const;
+
+	/**
+	 * Compares the filter bar's text with the item's component name. Use
+	 * bRecursive to refresh the state of child nodes as well. Returns true if
+	 * the node is set to be filtered out
+	 */
+	bool RefreshFilteredState(FSubobjectEditorTreeNodePtrType TreeNode, bool bRecursive);
+
+	/** Callback for the action trees to get the filter text */
+	FText GetFilterText() const;
+
+	/** Converts the current actor instance to a blueprint */
+	void PromoteToBlueprint() const;
+
+	/** Called when the promote to blueprint button is clicked */
+	FReply OnPromoteToBlueprintClicked();
+
+	/** Opens the blueprint editor for the blueprint being viewed by the scseditor */
+	void OnOpenBlueprintEditor(bool bForceCodeEditing) const;
+
+	/**
+	 * Spawns a new SWindow giving the user options for creating a new C++ component class.
+	 * The user will be prompted to pick a new subclass name and code will be recompiled
+	 *
+	 * @return The new class that was created
+	 */
+	UClass* SpawnCreateNewCppComponentWindow(TSubclassOf<UActorComponent> ComponentClass);
+
+	/**
+	 * Spawns a new SWindow giving the user options for creating a new blueprint component class.
+	 * The user will be prompted to pick a new subclass name and a blueprint asset will be created
+	 *
+	 * @return The new class that was created
+	 */
+	UClass* SpawnCreateNewBPComponentWindow(TSubclassOf<UActorComponent> ComponentClass);
+
+	/** Add a component from the selection in the combo box */
+	FSubobjectDataHandle PerformComboAddClass(TSubclassOf<UActorComponent> ComponentClass, EComponentCreateAction::Type ComponentCreateAction, UObject* AssetOverride);
+
+	/** @return The visibility of the promote to blueprint button (only visible with an actor instance that is not created from a blueprint)*/
+	virtual EVisibility GetPromoteToBlueprintButtonVisibility() const;
+
+	/** @return The visibility of the Edit Blueprint button (only visible with an actor instance that is created from a blueprint)*/
+	virtual EVisibility GetEditBlueprintButtonVisibility() const;
+
+	/** If true, then the tree widget will clear selection on click. */
+	virtual bool ClearSelectionOnClick() const { return false; }
+	
+	/** Indicates whether or not the search bar and inline buttons are visible. */
+	virtual bool ShowInlineSearchWithButtons() const { return false; }
+
+	/** Add a new subobject to the given parent via the subobject data subsystem */
+	virtual FSubobjectDataHandle AddNewSubobject(const FSubobjectDataHandle& ParentHandle, UClass* NewClass, UObject* AssetOverride, FText& OutFailReason, TUniquePtr<FScopedTransaction> InOngoingTransaction) = 0;
+
+	///////////////////////////////////////
+	// Utils
+	struct Utils
+	{
+		/** Populate an array of subobject handles based on an array of the slate node ptr */
+		static void PopulateHandlesArray(const TArray<FSubobjectEditorTreeNodePtrType>& SlateNodePtrs, TArray<FSubobjectDataHandle>& OutHandles);
+	};
+};
\ No newline at end of file
diff --git a/Engine/Source/Editor/SubobjectEditor/Public/SSubobjectEditorModule.h b/Engine/Source/Editor/SubobjectEditor/Public/SSubobjectEditorModule.h
new file mode 100644
--- /dev/null
+++ b/Engine/Source/Editor/SubobjectEditor/Public/SSubobjectEditorModule.h
@@ -0,0 +1,10 @@
+// Copyright Epic Games, Inc. All Rights Reserved.
+
+#pragma once
+
+#include "Modules/ModuleManager.h"
+
+class FSubobjectEditorModule : public IModuleInterface
+{
+
+};
\ No newline at end of file
diff --git a/Engine/Source/Editor/SubobjectEditor/Public/SSubobjectInstanceEditor.h b/Engine/Source/Editor/SubobjectEditor/Public/SSubobjectInstanceEditor.h
new file mode 100644
--- /dev/null
+++ b/Engine/Source/Editor/SubobjectEditor/Public/SSubobjectInstanceEditor.h
@@ -0,0 +1,55 @@
+// Copyright Epic Games, Inc. All Rights Reserved.
+
+#pragma once
+
+#include "SSubobjectEditor.h"
+
+/**
+* This is the editor for subobjects within the level editor that
+* works mainly with component and actor instances. 
+*/
+class SUBOBJECTEDITOR_API SSubobjectInstanceEditor final : public SSubobjectEditor
+{
+private:	
+	SLATE_BEGIN_ARGS(SSubobjectInstanceEditor)
+        : _ObjectContext(nullptr)
+        , _AllowEditing(true)
+        , _OnSelectionUpdated()
+		{}
+
+		SLATE_ATTRIBUTE(UObject*, ObjectContext)
+	    SLATE_ATTRIBUTE(bool, AllowEditing)
+	    SLATE_EVENT(FOnSelectionUpdated, OnSelectionUpdated)
+	    SLATE_EVENT(FOnItemDoubleClicked, OnItemDoubleClicked)
+		SLATE_EVENT(FSimpleDelegate, OnObjectReplaced)
+	
+	SLATE_END_ARGS()
+	
+	void Construct(const FArguments& InArgs);
+
+protected:
+	
+	// SSubobjectEditor interface
+	virtual void OnDeleteNodes() override;
+	virtual void CopySelectedNodes() override;
+	virtual void OnDuplicateComponent() override;
+	virtual void PasteNodes() override;
+	
+	virtual void OnAttachToDropAction(FSubobjectEditorTreeNodePtrType DroppedOn, const TArray<FSubobjectEditorTreeNodePtrType>& DroppedNodePtrs) override;
+    virtual void OnDetachFromDropAction(const TArray<FSubobjectEditorTreeNodePtrType>& DroppedNodePtrs) override;
+    virtual void OnMakeNewRootDropAction(FSubobjectEditorTreeNodePtrType DroppedNodePtr) override;
+    virtual void PostDragDropAction(bool bRegenerateTreeNodes) override;
+
+    /** Builds a context menu popup for dropping a child node onto the scene root node */
+    virtual TSharedPtr<SWidget> BuildSceneRootDropActionMenu(FSubobjectEditorTreeNodePtrType DroppedOntoNodePtr, FSubobjectEditorTreeNodePtrType DroppedNodePtr) override;
+	virtual FSubobjectDataHandle AddNewSubobject(const FSubobjectDataHandle& ParentHandle, UClass* NewClass, UObject* AssetOverride, FText& OutFailReason, TUniquePtr<FScopedTransaction> InOngoingTransaction) override;
+	virtual void PopulateContextMenuImpl(UToolMenu* InMenu, TArray<FSubobjectEditorTreeNodePtrType>& InSelectedItems, bool bIsChildActorSubtreeNodeSelected) override;
+
+public:
+	virtual FSlateColor GetColorTintForIcon(FSubobjectEditorTreeNodePtrType Node) const override;
+	// End of SSubobjectEditor
+
+	/** Delegate to invoke when objects within the Subobject tree are replaced (eg, via re-instancing from a BP compile) */
+	FSimpleDelegate OnObjectReplaced;
+
+};
\ No newline at end of file
diff --git a/Engine/Source/Editor/SubobjectEditor/SubobjectEditor.Build.cs b/Engine/Source/Editor/SubobjectEditor/SubobjectEditor.Build.cs
new file mode 100644
--- /dev/null
+++ b/Engine/Source/Editor/SubobjectEditor/SubobjectEditor.Build.cs
@@ -0,0 +1,46 @@
+// Copyright Epic Games, Inc. All Rights Reserved.
+
+using UnrealBuildTool;
+using System.IO;
+
+public class SubobjectEditor : ModuleRules
+{
+	public SubobjectEditor(ReadOnlyTargetRules Target) : base(Target)
+	{
+		PublicIncludePaths.Add(Path.Combine(ModuleDirectory, "Public"));
+		PrivateIncludePaths.Add(Path.Combine(ModuleDirectory, "Private"));
+
+		PublicDependencyModuleNames.AddRange(
+			new string[]
+			{
+				"Core",
+				"CoreUObject",
+				"SubobjectDataInterface",
+				"Kismet",
+			}
+		);
+
+		PrivateDependencyModuleNames.AddRange(
+			new string[]
+			{
+				"BlueprintGraph",
+				"Engine",
+				"Slate",
+				"SlateCore",
+				"EditorStyle",
+				"UnrealEd",
+				"InputCore",
+				"KismetWidgets",
+				"GameProjectGeneration",
+				"EditorFramework",
+				"GraphEditor",
+				"ToolMenus",
+			}
+		);
+		
+		if(Target.bWithLiveCoding)
+		{
+			PrivateIncludePathModuleNames.Add("LiveCoding");
+		}
+	}
+}
