Update implementation of TryHandleAssetDragDropOperation

#rnx

diff --git a/Engine/Source/Editor/SubobjectEditor/Private/SSubobjectEditor.cpp b/Engine/Source/Editor/SubobjectEditor/Private/SSubobjectEditor.cpp
--- a/Engine/Source/Editor/SubobjectEditor/Private/SSubobjectEditor.cpp
+++ b/Engine/Source/Editor/SubobjectEditor/Private/SSubobjectEditor.cpp
@@ -53,6 +53,7 @@
 #include "K2Node_VariableGet.h"
 #include "AssetSelection.h"
 #include "Settings/EditorProjectSettings.h"
+#include "Misc/FeedbackContext.h"
 
 extern UNREALED_API UEditorEngine* GEditor;
 
@@ -2469,8 +2470,116 @@ FReply SSubobjectEditor::OnKeyDown(const FGeometry& MyGeometry, const FKeyEvent&
 
 FReply SSubobjectEditor::TryHandleAssetDragDropOperation(const FDragDropEvent& DragDropEvent)
 {
-	// #TODO_BH Implement
-	ensureMsgf(false, TEXT("Implement TryHandleAssetDragDropOperation!"));
+	TSharedPtr<FDragDropOperation> Operation = DragDropEvent.GetOperation();
+	if (Operation.IsValid() && (Operation->IsOfType<FExternalDragOperation>() || Operation->IsOfType<FAssetDragDropOp>()))
+	{
+		TArray< FAssetData > DroppedAssetData = AssetUtil::ExtractAssetDataFromDrag(Operation);
+		const int32 NumAssets = DroppedAssetData.Num();
+
+		if (NumAssets > 0)
+		{
+			GWarn->BeginSlowTask(LOCTEXT("LoadingAssets", "Loading Asset(s)"), true);
+			bool bMarkBlueprintAsModified = false;
+
+			for (int32 DroppedAssetIdx = 0; DroppedAssetIdx < NumAssets; ++DroppedAssetIdx)
+			{
+				const FAssetData& AssetData = DroppedAssetData[DroppedAssetIdx];
+
+				if (!AssetData.IsAssetLoaded())
+				{
+					GWarn->StatusUpdate(DroppedAssetIdx, NumAssets, FText::Format(LOCTEXT("LoadingAsset", "Loading Asset {0}"), FText::FromName(AssetData.AssetName)));
+				}
+				
+				UClass* AssetClass = AssetData.GetClass();
+				UObject* Asset = AssetData.GetAsset();
+
+				UBlueprint* BPClass = Cast<UBlueprint>(Asset);
+				UClass* PotentialComponentClass = nullptr;
+				UClass* PotentialActorClass = nullptr;
+
+				if ((BPClass != nullptr) && (BPClass->GeneratedClass != nullptr))
+				{
+					if (BPClass->GeneratedClass->IsChildOf(UActorComponent::StaticClass()))
+					{
+						PotentialComponentClass = BPClass->GeneratedClass;
+					}
+					else if (BPClass->GeneratedClass->IsChildOf(AActor::StaticClass()))
+					{
+						PotentialActorClass = BPClass->GeneratedClass;
+					}
+				}
+				else if (AssetClass->IsChildOf(UClass::StaticClass()))
+				{
+					UClass* AssetAsClass = CastChecked<UClass>(Asset);
+					if (AssetAsClass->IsChildOf(UActorComponent::StaticClass()))
+					{
+						PotentialComponentClass = AssetAsClass;
+					}
+					else if (AssetAsClass->IsChildOf(AActor::StaticClass()))
+					{
+						PotentialActorClass = AssetAsClass;
+					}
+				}
+				
+				USubobjectDataSubsystem* System = USubobjectDataSubsystem::Get();
+				check(System);
+
+				FAddNewSubobjectParams NewComponentParams;
+				FText FailReason;
+				NewComponentParams.bSkipMarkBlueprintModified = true;
+				const bool bSetFocusToNewItem = (DroppedAssetIdx == NumAssets - 1);// Only set focus to the last item created
+				
+				TSubclassOf<UActorComponent> MatchingComponentClassForAsset = FComponentAssetBrokerage::GetPrimaryComponentForAsset(AssetClass);
+				if (MatchingComponentClassForAsset != nullptr)
+				{
+					NewComponentParams.NewClass = MatchingComponentClassForAsset;
+					NewComponentParams.AssetOverride = Asset;
+					bMarkBlueprintAsModified = true;
+					System->AddNewSubobject(NewComponentParams, FailReason);
+				}
+				else if ((PotentialComponentClass != nullptr) && !PotentialComponentClass->HasAnyClassFlags(CLASS_Deprecated | CLASS_Abstract | CLASS_NewerVersionExists))
+				{
+					if (PotentialComponentClass->HasMetaData(FBlueprintMetadata::MD_BlueprintSpawnableComponent))
+					{
+						NewComponentParams.NewClass = PotentialComponentClass;
+						NewComponentParams.AssetOverride = nullptr;
+						bMarkBlueprintAsModified = true;
+						System->AddNewSubobject(NewComponentParams, FailReason);
+					}
+				}
+				else if ((PotentialActorClass != nullptr) && !PotentialActorClass->HasAnyClassFlags(CLASS_Deprecated | CLASS_Abstract | CLASS_NewerVersionExists | CLASS_NotPlaceable))
+				{
+					NewComponentParams.AssetOverride = UChildActorComponent::StaticClass();
+					NewComponentParams.AssetOverride = PotentialActorClass;					
+					bMarkBlueprintAsModified = true;
+					System->AddNewSubobject(NewComponentParams, FailReason);
+				}
+			}
+			
+			// Optimization: Only mark the blueprint as modified at the end
+			if (bMarkBlueprintAsModified && ShouldModifyBPOnAssetDrop())
+			{
+				UBlueprint* Blueprint = GetBlueprint();
+				check(Blueprint != nullptr && Blueprint->SimpleConstructionScript != nullptr);
+
+				Blueprint->Modify();
+				if (Blueprint->SimpleConstructionScript)
+				{
+					Blueprint->SimpleConstructionScript->SaveToTransactionBuffer();
+				}
+
+				bAllowTreeUpdates = true;
+				FBlueprintEditorUtils::MarkBlueprintAsStructurallyModified(Blueprint);
+			}
+
+			UpdateTree();
+
+			GWarn->EndSlowTask();			
+		}
+		
+		return FReply::Handled();
+	}
+	
 	return FReply::Unhandled();
 }
 
diff --git a/Engine/Source/Editor/SubobjectEditor/Public/SSubobjectBlueprintEditor.h b/Engine/Source/Editor/SubobjectEditor/Public/SSubobjectBlueprintEditor.h
--- a/Engine/Source/Editor/SubobjectEditor/Public/SSubobjectBlueprintEditor.h
+++ b/Engine/Source/Editor/SubobjectEditor/Public/SSubobjectBlueprintEditor.h
@@ -33,6 +33,7 @@ protected:
 	TAttribute<AActor*> PreviewActor;
 	
 	// SSubobjectEditor interface
+	virtual bool ShouldModifyBPOnAssetDrop() const override { return true; }
 	virtual void OnDeleteNodes() override;
 	virtual void CopySelectedNodes() override;
 	virtual void OnDuplicateComponent() override;
diff --git a/Engine/Source/Editor/SubobjectEditor/Public/SSubobjectEditor.h b/Engine/Source/Editor/SubobjectEditor/Public/SSubobjectEditor.h
--- a/Engine/Source/Editor/SubobjectEditor/Public/SSubobjectEditor.h
+++ b/Engine/Source/Editor/SubobjectEditor/Public/SSubobjectEditor.h
@@ -541,6 +541,9 @@ protected:
 
 	/** Restore the previous selection state when updating the tree */
 	virtual void RestoreSelectionState(TArray<FSubobjectEditorTreeNodePtrType>& SelectedTreeNodes);
+
+	/** If true, then the blueprint should be modified on TryHandleAssetDragDropOperation */
+	virtual bool ShouldModifyBPOnAssetDrop() const { return false; }
 	
 	bool CanCutNodes() const;
 	void CutSelectedNodes();
