Fix non-unity compile error in SSubobjectInstanceEditor.cpp

#rnx

diff --git a/Engine/Source/Editor/SubobjectEditor/Private/SSubobjectBlueprintEditor.cpp b/Engine/Source/Editor/SubobjectEditor/Private/SSubobjectBlueprintEditor.cpp
--- a/Engine/Source/Editor/SubobjectEditor/Private/SSubobjectBlueprintEditor.cpp
+++ b/Engine/Source/Editor/SubobjectEditor/Private/SSubobjectBlueprintEditor.cpp
@@ -3,6 +3,7 @@
 #include "SSubobjectBlueprintEditor.h"
 #include "ScopedTransaction.h"
 #include "IDocumentation.h"
+#include "Widgets/SToolTip.h"
 #include "GraphEditorActions.h"
 #include "Editor/EditorEngine.h"
 #include "ISCSEditorUICustomization.h"	// #TODO_BH Rename this to subobject
diff --git a/Engine/Source/Editor/SubobjectEditor/Private/SSubobjectInstanceEditor.cpp b/Engine/Source/Editor/SubobjectEditor/Private/SSubobjectInstanceEditor.cpp
--- a/Engine/Source/Editor/SubobjectEditor/Private/SSubobjectInstanceEditor.cpp
+++ b/Engine/Source/Editor/SubobjectEditor/Private/SSubobjectInstanceEditor.cpp
@@ -3,6 +3,7 @@
 #include "SSubobjectInstanceEditor.h"
 #include "ScopedTransaction.h"
 #include "IDocumentation.h"
+#include "Widgets/SToolTip.h"
 #include "GraphEditorActions.h"
 #include "Editor/EditorEngine.h"
 #include "ISCSEditorUICustomization.h"	// #TODO_BH Rename this to subobject
@@ -75,7 +76,7 @@ void SSubobjectInstanceEditor::Construct(const FArguments& InArgs)
 	]
 	+ SHorizontalBox::Slot()
 	.VAlign(VAlign_Center)
-	.Padding(0,0,4,0)
+	.Padding(0.0f,0.0f,4.0f,0.0f)
 	.AutoWidth()
 	[
 		SAssignNew(ExtensionPanel, SExtensionPanel)
@@ -87,7 +88,7 @@ void SSubobjectInstanceEditor::Construct(const FArguments& InArgs)
 
 	+ SHorizontalBox::Slot()
 	.VAlign(VAlign_Center)
-	.Padding(0,0,4,0)
+	.Padding(0.0f,0.0f,4.0f,0.0f)
 	.AutoWidth()
 	[
 		SNew(SEditorHeaderButton)
