Add CAC visualization mode for the subobject editor

#rnx

diff --git a/Engine/Source/Editor/SubobjectEditor/Private/SSubobjectEditor.cpp b/Engine/Source/Editor/SubobjectEditor/Private/SSubobjectEditor.cpp
--- a/Engine/Source/Editor/SubobjectEditor/Private/SSubobjectEditor.cpp
+++ b/Engine/Source/Editor/SubobjectEditor/Private/SSubobjectEditor.cpp
@@ -136,21 +136,20 @@ namespace SubobjectEditorUtils
 				WeakEditorPtr = MenuContext->SubobjectEditor;
 			}
 
-			// #TODO_BH Why does this give me a linker error???????????
-			// Section.AddMenuEntry(
-   //              "SetChildActorOnlyMode",
-   //              LOCTEXT("SetChildActorOnlyMode_Label", "Switch to Child Actor Only Mode"),
-   //              LOCTEXT("SetChildActorOnlyMode_ToolTip", "Visualize this child actor's template/instance subtree in place of its parent component node."),
-   //              FSlateIcon(),
-   //              FUIAction(
-   //                  FExecuteAction::CreateStatic(
-	  //                   &SubobjectEditorUtils::OnSetChildActorTreeViewVisualizationMode, 
-	  //                   ChildActorComponent,
-	  //                   EChildActorComponentTreeViewVisualizationMode::ChildActorOnly,
-	  //                   WeakEditorPtr),
-   //                  FCanExecuteAction()
-   //              )
-   //          );
+			Section.AddMenuEntry(
+				"SetChildActorOnlyMode",
+				LOCTEXT("SetChildActorOnlyMode_Label", "Switch to Child Actor Only Mode"),
+				LOCTEXT("SetChildActorOnlyMode_ToolTip", "Visualize this child actor's template/instance subtree in place of its parent component node."),
+				FSlateIcon(),
+				FUIAction(
+					FExecuteAction::CreateStatic(
+						&SubobjectEditorUtils::OnSetChildActorTreeViewVisualizationMode,
+						ChildActorComponent,
+						EChildActorComponentTreeViewVisualizationMode::ChildActorOnly,
+						WeakEditorPtr),
+					FCanExecuteAction()
+				)
+			);
 		}
 	}
 }
