Regenerate tree structure upon deteaching a subobject in the instance and blueprint subobject editors

#rnx

diff --git a/Engine/Source/Editor/SubobjectEditor/Private/SSubobjectBlueprintEditor.cpp b/Engine/Source/Editor/SubobjectEditor/Private/SSubobjectBlueprintEditor.cpp
--- a/Engine/Source/Editor/SubobjectEditor/Private/SSubobjectBlueprintEditor.cpp
+++ b/Engine/Source/Editor/SubobjectEditor/Private/SSubobjectBlueprintEditor.cpp
@@ -411,7 +411,7 @@ void SSubobjectBlueprintEditor::OnDetachFromDropAction(const TArray<FSubobjectEd
 		System->ReparentSubobjects(Params, HandlesToMove);
 	}
 
-	PostDragDropAction(false);
+	PostDragDropAction(true);
 }
 
 void SSubobjectBlueprintEditor::OnMakeNewRootDropAction(FSubobjectEditorTreeNodePtrType DroppedNodePtr)
diff --git a/Engine/Source/Editor/SubobjectEditor/Private/SSubobjectInstanceEditor.cpp b/Engine/Source/Editor/SubobjectEditor/Private/SSubobjectInstanceEditor.cpp
--- a/Engine/Source/Editor/SubobjectEditor/Private/SSubobjectInstanceEditor.cpp
+++ b/Engine/Source/Editor/SubobjectEditor/Private/SSubobjectInstanceEditor.cpp
@@ -308,19 +308,13 @@ void SSubobjectInstanceEditor::OnDetachFromDropAction(const TArray<FSubobjectEdi
 		System->ReparentSubobjects(Params, HandlesToMove);
 	}
 
-	PostDragDropAction(false);
+	PostDragDropAction(true);
 }
 
 void SSubobjectInstanceEditor::OnMakeNewRootDropAction(FSubobjectEditorTreeNodePtrType DroppedNodePtr)
 {
 	// Get the current scene root node
 	FSubobjectEditorTreeNodePtrType SceneRootNodePtr = GetSceneRootNode();
-
-	// We cannot handle the drop action if any of these conditions fail on entry.
-	//if (!ensure(SceneRootNodePtr.IsValid()) || !ensure(DroppedNodePtr.IsValid()) || !ensure(DroppedNodePtr == SceneRootNodePtr))
-	//{
-	//	return;
-	//}
 	
 	// Create a transaction record
 	const FScopedTransaction TransactionContext(LOCTEXT("MakeNewSceneRoot", "Make New Scene Root"));
