Add IsInheritedComponent override to FInheritedSubobjectData to get it to correctly return true.

#rnx

diff --git a/Engine/Source/Editor/SubobjectDataInterface/Private/InheritedSubobjectData.cpp b/Engine/Source/Editor/SubobjectDataInterface/Private/InheritedSubobjectData.cpp
--- a/Engine/Source/Editor/SubobjectDataInterface/Private/InheritedSubobjectData.cpp
+++ b/Engine/Source/Editor/SubobjectDataInterface/Private/InheritedSubobjectData.cpp
@@ -42,6 +42,25 @@ bool FInheritedSubobjectData::CanEdit() const
 	return FSubobjectData::CanEdit();
 }
 
+bool FInheritedSubobjectData::CanDelete() const
+{
+	if(IsNativeComponent() || IsInheritedSCSNode() || IsDefaultSceneRoot() || IsInstancedInheritedComponent())
+	{
+		return false;
+	}
+	
+	return true;
+}
+
+bool FInheritedSubobjectData::IsInheritedComponent() const
+{
+	if(GetComponentTemplate() != nullptr)
+	{
+		return true;
+	}
+	return false;
+}
+
 bool FInheritedSubobjectData::IsInheritedSCSNode() const
 {
 	return bIsInheritedSCS;
diff --git a/Engine/Source/Editor/SubobjectDataInterface/Public/InheritedSubobjectData.h b/Engine/Source/Editor/SubobjectDataInterface/Public/InheritedSubobjectData.h
--- a/Engine/Source/Editor/SubobjectDataInterface/Public/InheritedSubobjectData.h
+++ b/Engine/Source/Editor/SubobjectDataInterface/Public/InheritedSubobjectData.h
@@ -5,7 +5,7 @@
 #include "SubobjectData.h"
 
 /**
- * Container that represents the subobejct data of a child actor component.
+ * Container that represents the subobject data of a child actor component.
  */
 struct SUBOBJECTDATAINTERFACE_API FInheritedSubobjectData : public FSubobjectData
 {
@@ -16,7 +16,9 @@ struct SUBOBJECTDATAINTERFACE_API FInheritedSubobjectData : public FSubobjectDat
 	// FSubobjectData interface
 	virtual bool IsNativeComponent() const override;
 	virtual bool CanEdit() const override;
+	virtual bool CanDelete() const override;
 	virtual bool IsInheritedSCSNode() const override;
+	virtual bool IsInheritedComponent() const override;
 	// End FSubobjectData
 
 protected:
diff --git a/Engine/Source/Editor/SubobjectEditor/Private/SSubobjectEditor.cpp b/Engine/Source/Editor/SubobjectEditor/Private/SSubobjectEditor.cpp
--- a/Engine/Source/Editor/SubobjectEditor/Private/SSubobjectEditor.cpp
+++ b/Engine/Source/Editor/SubobjectEditor/Private/SSubobjectEditor.cpp
@@ -2139,19 +2139,9 @@ FSubobjectEditorTreeNodePtrType SSubobjectEditor::GetSceneRootNode() const
 		if(const USceneComponent* RootComp = ActorContext->GetRootComponent())
 		{
 			TSet<FSubobjectEditorTreeNodePtrType> VisitedNodes;
-			FSubobjectDataHandle RootHandle;
+			USubobjectDataSubsystem* System = USubobjectDataSubsystem::Get();
 			
-			DepthFirstTraversal(RootNodes[0], VisitedNodes,
-                [&RootComp, &RootHandle](
-                const FSubobjectEditorTreeNodePtrType& InNodePtr)
-                {
-                    if(InNodePtr->GetComponentTemplate() == RootComp)
-                    {
-                    	RootHandle = InNodePtr->GetDataHandle();
-                    }
-                });
-			
-			return FindSlateNodeForHandle(RootHandle);
+			return FindSlateNodeForHandle(System->FindSceneRootForSubobject(GetObjectContextHandle()));
 		}
 	}
 
