Override FindSlateNodeForObject in the Blueprint editor to handle getting the archetype of the actor compoent

#rnx

diff --git a/Engine/Source/Editor/SubobjectEditor/Private/SSubobjectBlueprintEditor.cpp b/Engine/Source/Editor/SubobjectEditor/Private/SSubobjectBlueprintEditor.cpp
--- a/Engine/Source/Editor/SubobjectEditor/Private/SSubobjectBlueprintEditor.cpp
+++ b/Engine/Source/Editor/SubobjectEditor/Private/SSubobjectBlueprintEditor.cpp
@@ -611,6 +611,68 @@ FSubobjectEditorTreeNodePtrType SSubobjectBlueprintEditor::GetSceneRootNode() co
 	return nullptr;
 }
 
+FSubobjectEditorTreeNodePtrType SSubobjectBlueprintEditor::FindSlateNodeForObject(const UObject* InObject, bool bIncludeAttachmentComponents) const
+{
+	if (const UActorComponent* ActorComponent = Cast<UActorComponent>(InObject))
+	{
+		// If the given component instance is not already an archetype object
+		if (!ActorComponent->IsTemplate())
+		{
+			// Get the component owner's class object
+			check(ActorComponent->GetOwner() != NULL);
+			UClass* OwnerClass = ActorComponent->GetOwner()->GetClass();
+
+			// If the given component is one that's created during Blueprint construction
+			if (ActorComponent->IsCreatedByConstructionScript())
+			{
+				// Check the entire Class hierarchy for the node
+				TArray<UBlueprintGeneratedClass*> ParentBPStack;
+				UBlueprint::GetBlueprintHierarchyFromClass(OwnerClass, ParentBPStack);
+
+				for (int32 StackIndex = ParentBPStack.Num() - 1; StackIndex >= 0; --StackIndex)
+				{
+					USimpleConstructionScript* ParentSCS = ParentBPStack[StackIndex] ? ParentBPStack[StackIndex]->SimpleConstructionScript.Get() : nullptr;
+					if (ParentSCS)
+					{
+						// Attempt to locate an SCS node with a variable name that matches the name of the given component
+						for (USCS_Node* SCS_Node : ParentSCS->GetAllNodes())
+						{
+							check(SCS_Node != nullptr);
+							if (SCS_Node->GetVariableName() == ActorComponent->GetFName())
+							{
+								// We found a match; redirect to the component archetype instance that may be associated with a tree node
+								ActorComponent = SCS_Node->ComponentTemplate;
+								break;
+							}
+						}
+					}
+				}
+			}
+			else
+			{
+				// Get the class default object
+				const AActor* CDO = Cast<AActor>(OwnerClass->GetDefaultObject());
+				if (CDO)
+				{
+					// Iterate over the Components array and attempt to find a component with a matching name
+					for (UActorComponent* ComponentTemplate : CDO->GetComponents())
+					{
+						if (ComponentTemplate && ComponentTemplate->GetFName() == ActorComponent->GetFName())
+						{
+							// We found a match; redirect to the component archetype instance that may be associated with a tree node
+							ActorComponent = ComponentTemplate;
+							break;
+						}
+					}
+				}
+			}
+		}
+		InObject = ActorComponent;
+	}
+	
+	return SSubobjectEditor::FindSlateNodeForObject(InObject, bIncludeAttachmentComponents);
+}
+
 void SSubobjectBlueprintEditor::BuildMenuEventsSection(FMenuBuilder& Menu, UBlueprint* Blueprint, UClass* SelectedClass, FCanExecuteAction CanExecuteActionDelegate, FGetSelectedObjectsDelegate GetSelectedObjectsDelegate)
 {
 	// Get Selected Nodes
diff --git a/Engine/Source/Editor/SubobjectEditor/Private/SSubobjectEditor.cpp b/Engine/Source/Editor/SubobjectEditor/Private/SSubobjectEditor.cpp
--- a/Engine/Source/Editor/SubobjectEditor/Private/SSubobjectEditor.cpp
+++ b/Engine/Source/Editor/SubobjectEditor/Private/SSubobjectEditor.cpp
@@ -2839,6 +2839,16 @@ FSubobjectEditorTreeNodePtrType SSubobjectEditor::FindSlateNodeForObject(const U
                     }
                 }
             });
+
+		// If we didn't find it in the tree, step up the chain to the parent of the given component and recursively see if that is in the tree (unless the flag is false)
+		if (!OutNodePtr.IsValid() && bIncludeAttachmentComponents)
+		{
+			const USceneComponent* SceneComponent = Cast<const USceneComponent>(InObject);
+			if (SceneComponent && SceneComponent->GetAttachParent())
+			{
+				return FindSlateNodeForObject(SceneComponent->GetAttachParent(), bIncludeAttachmentComponents);
+			}
+		}
 	}
 
 	return OutNodePtr;
diff --git a/Engine/Source/Editor/SubobjectEditor/Public/SSubobjectBlueprintEditor.h b/Engine/Source/Editor/SubobjectEditor/Public/SSubobjectBlueprintEditor.h
--- a/Engine/Source/Editor/SubobjectEditor/Public/SSubobjectBlueprintEditor.h
+++ b/Engine/Source/Editor/SubobjectEditor/Public/SSubobjectBlueprintEditor.h
@@ -53,7 +53,7 @@ protected:
 	virtual void PopulateContextMenuImpl(UToolMenu* InMenu, TArray<FSubobjectEditorTreeNodePtrType>& InSelectedItems, bool bIsChildActorSubtreeNodeSelected) override;
 	virtual FSubobjectEditorTreeNodePtrType GetSceneRootNode() const override;
 public:
-	//virtual FSubobjectEditorTreeNodePtrType FindSlateNodeForObject(const UObject* InObject, bool bIncludeAttachmentComponents = true) const override;
+	virtual FSubobjectEditorTreeNodePtrType FindSlateNodeForObject(const UObject* InObject, bool bIncludeAttachmentComponents = true) const override;
 	// End of SSubobjectEditor
 
 public:
