make SSubobjectEditor::GetObjectContext and SSubobjectEditor::GetObjectContextHandle public accessible functions for users to extend

#rnx

diff --git a/Engine/Source/Editor/SubobjectEditor/Public/SSubobjectEditor.h b/Engine/Source/Editor/SubobjectEditor/Public/SSubobjectEditor.h
--- a/Engine/Source/Editor/SubobjectEditor/Public/SSubobjectEditor.h
+++ b/Engine/Source/Editor/SubobjectEditor/Public/SSubobjectEditor.h
@@ -539,6 +539,12 @@ public:
 
 	// Attempt to find an existing slate node that matches the given handle
 	FSubobjectEditorTreeNodePtrType FindSlateNodeForHandle(const FSubobjectDataHandle& Handle, FSubobjectEditorTreeNodePtrType InStartNodePtr = FSubobjectEditorTreeNodePtrType()) const;
+
+	/** Pointer to the current object that is represented by the subobject editor */
+	UObject* GetObjectContext() const;
+
+	/** SubobjectHandle of the current object that is represented by the subobject editor */
+	FSubobjectDataHandle GetObjectContextHandle() const;
 	
 protected:
 
@@ -628,10 +634,6 @@ protected:
 
 	/** Creates a list of commands */
 	virtual void CreateCommandList();
-	
-	UObject* GetObjectContext() const;
-
-	FSubobjectDataHandle GetObjectContextHandle() const;
 
 	/** Recursively visits the given node + its children and invokes the given function for each. */
 	void DepthFirstTraversal(const FSubobjectEditorTreeNodePtrType& InNodePtr, TSet<FSubobjectEditorTreeNodePtrType>& OutVisitedNodes, const TFunctionRef<void(const FSubobjectEditorTreeNodePtrType&)> InFunction) const;
