Keep track of the created slate nodes while updating tree in the local function scope instead of calling FindOrCreateSlateNodeForHandle because we cannot gaurantee that the grandparent node will be in the RootTree when creating subobjects.

#rnx

diff --git a/Engine/Source/Editor/SubobjectDataInterface/Public/SubobjectDataHandle.h b/Engine/Source/Editor/SubobjectDataInterface/Public/SubobjectDataHandle.h
--- a/Engine/Source/Editor/SubobjectDataInterface/Public/SubobjectDataHandle.h
+++ b/Engine/Source/Editor/SubobjectDataInterface/Public/SubobjectDataHandle.h
@@ -58,6 +58,12 @@ struct SUBOBJECTDATAINTERFACE_API FSubobjectDataHandle
 	*/
 	static const FSubobjectDataHandle InvalidHandle;
 	
+	/** Get the hash code to use for the given FSubobjectDataHandle */
+	friend uint32 GetTypeHash(const FSubobjectDataHandle& Key)
+	{
+		return PointerHash(Key.DataPtr.Get());
+	}
+	
 private:
 	/** Pointer to the actual subobject data that this handle represents */
 	TSharedPtr<FSubobjectData> DataPtr;
diff --git a/Engine/Source/Editor/SubobjectEditor/Private/SSubobjectEditor.cpp b/Engine/Source/Editor/SubobjectEditor/Private/SSubobjectEditor.cpp
--- a/Engine/Source/Editor/SubobjectEditor/Private/SSubobjectEditor.cpp
+++ b/Engine/Source/Editor/SubobjectEditor/Private/SSubobjectEditor.cpp
@@ -2344,12 +2344,15 @@ void SSubobjectEditor::UpdateTree(bool bRegenerateTreeNodes /* = true */)
 		DataSubsystem->GatherSubobjectData(Context, SubobjectData);
 		
 		FSubobjectEditorTreeNodePtrType SeperatorNode;
+		TMap<FSubobjectDataHandle, FSubobjectEditorTreeNodePtrType> AddedNodes;
+
 		// By default, root node will always be expanded. If possible we will restore the collapsed state later.
 		if (SubobjectData.Num() > 0)
 		{
 			FSubobjectEditorTreeNodePtrType Node = MakeShareable<FSubobjectEditorTreeNode>(
 				new FSubobjectEditorTreeNode(SubobjectData[0]));
 			RootNodes.Add(Node);
+			AddedNodes.Add(Node->GetDataHandle(), Node);
 			CachedRootHandle = Node->GetDataHandle();
 
 			SeperatorNode = MakeShareable<FSubobjectEditorTreeNode>(
@@ -2363,10 +2366,11 @@ void SSubobjectEditor::UpdateTree(bool bRegenerateTreeNodes /* = true */)
 		}
 
 		// Create slate nodes for each subobject
-		for (const FSubobjectDataHandle& Handle : SubobjectData)
+		for (FSubobjectDataHandle& Handle : SubobjectData)
 		{
 			// Do we have a slate node for this handle already? If not, then we need to make one
-			FSubobjectEditorTreeNodePtrType NewNode = FindOrCreateSlateNodeForHandle(Handle);
+			FSubobjectEditorTreeNodePtrType NewNode = SSubobjectEditor::FindOrCreateSlateNodeForHandle(Handle, AddedNodes);
+			
 			FSubobjectData* Data = Handle.GetData();
 
 			const FSubobjectDataHandle& ParentHandle = Data->GetParentHandle();
@@ -2375,7 +2379,8 @@ void SSubobjectEditor::UpdateTree(bool bRegenerateTreeNodes /* = true */)
 			if (ParentHandle.IsValid())
 			{
 				// Get the parent node for this subobject
-				FSubobjectEditorTreeNodePtrType ParentNode = FindOrCreateSlateNodeForHandle(ParentHandle);
+				FSubobjectEditorTreeNodePtrType ParentNode = SSubobjectEditor::FindOrCreateSlateNodeForHandle(ParentHandle, AddedNodes);
+
 				check(ParentNode);
 				ParentNode->AddChild(NewNode);
 				TreeWidget->SetItemExpansion(ParentNode, true);			
@@ -2399,6 +2404,21 @@ void SSubobjectEditor::UpdateTree(bool bRegenerateTreeNodes /* = true */)
 	TreeWidget->RequestTreeRefresh();
 }
 
+FSubobjectEditorTreeNodePtrType SSubobjectEditor::FindOrCreateSlateNodeForHandle(const FSubobjectDataHandle& InHandle, TMap<FSubobjectDataHandle, FSubobjectEditorTreeNodePtrType>& ExistingNodes)
+{
+	// If we have already created a tree node for this handle, great! Return that.
+	if (const FSubobjectEditorTreeNodePtrType* Found = ExistingNodes.Find(InHandle))
+	{
+		return *Found;
+	}
+
+	// Otherwise we haven't created this yet and we need to make a new one!
+	FSubobjectEditorTreeNodePtrType NewNode = MakeShareable<FSubobjectEditorTreeNode>(new FSubobjectEditorTreeNode(InHandle));
+	ExistingNodes.Add(InHandle, NewNode);
+
+	return NewNode;
+}
+
 void SSubobjectEditor::RefreshSelectionDetails()
 {
 	UpdateSelectionFromNodes(TreeWidget->GetSelectedItems());
@@ -2854,18 +2874,6 @@ FSubobjectEditorTreeNodePtrType SSubobjectEditor::FindSlateNodeForObject(const U
 	return OutNodePtr;
 }
 
-FSubobjectEditorTreeNodePtrType SSubobjectEditor::FindOrCreateSlateNodeForHandle(const FSubobjectDataHandle& Handle) const
-{
-	if (FSubobjectEditorTreeNodePtrType Existing = FindSlateNodeForHandle(Handle))
-	{
-		return Existing;
-	}
-	else
-	{
-		return MakeShareable<FSubobjectEditorTreeNode>(new FSubobjectEditorTreeNode(Handle));
-	}
-}
-
 void SSubobjectEditor::SetNodeExpansionState(FSubobjectEditorTreeNodePtrType InNodeToChange, const bool bIsExpanded)
 {
 	if (TreeWidget.IsValid() && InNodeToChange.IsValid())
diff --git a/Engine/Source/Editor/SubobjectEditor/Public/SSubobjectEditor.h b/Engine/Source/Editor/SubobjectEditor/Public/SSubobjectEditor.h
--- a/Engine/Source/Editor/SubobjectEditor/Public/SSubobjectEditor.h
+++ b/Engine/Source/Editor/SubobjectEditor/Public/SSubobjectEditor.h
@@ -642,7 +642,7 @@ protected:
 	void GetCollapsedNodes(const FSubobjectEditorTreeNodePtrType& InNodePtr, TSet<FSubobjectEditorTreeNodePtrType>& OutCollapsedNodes) const;
 
 	// Attempt to find an existing slate node that matches the given handle
-	FSubobjectEditorTreeNodePtrType FindOrCreateSlateNodeForHandle(const FSubobjectDataHandle& Handle) const;
+	static FSubobjectEditorTreeNodePtrType FindOrCreateSlateNodeForHandle(const FSubobjectDataHandle& Handle, TMap<FSubobjectDataHandle, FSubobjectEditorTreeNodePtrType>& ExistingNodes);
 
 	/**
 	 * Set the expansion state of a node
