When deleting an instanced component, select the scene root in the tree if there was no given subobject to select

#rnx

diff --git a/Engine/Source/Editor/SubobjectEditor/Private/SSubobjectInstanceEditor.cpp b/Engine/Source/Editor/SubobjectEditor/Private/SSubobjectInstanceEditor.cpp
--- a/Engine/Source/Editor/SubobjectEditor/Private/SSubobjectInstanceEditor.cpp
+++ b/Engine/Source/Editor/SubobjectEditor/Private/SSubobjectInstanceEditor.cpp
@@ -193,14 +193,11 @@ void SSubobjectInstanceEditor::OnDeleteNodes()
 		int32 NumDeleted = System->DeleteSubobjects(RootNodes[0]->GetDataHandle(), HandlesToDelete, HandleToSelect);
 
 		if(NumDeleted > 0)
-		{
-			if(HandleToSelect.IsValid())
+		{			
+			FSubobjectEditorTreeNodePtrType NodeToSelect = HandleToSelect.IsValid() ? FindSlateNodeForHandle(HandleToSelect) : GetSceneRootNode();
+			if(NodeToSelect.IsValid())
 			{
-				FSubobjectEditorTreeNodePtrType NodeToSelect = FindSlateNodeForHandle(HandleToSelect);
-				if(NodeToSelect.IsValid())
-				{
-					TreeWidget->SetSelection(NodeToSelect);
-				}
+				TreeWidget->SetSelection(NodeToSelect);
 			}
 			
 			UpdateTree();
