Add a GetDragDropDisplayText function to subobject data to allow for easier conversions to FText for drag/drop operations

#rnx

diff --git a/Engine/Source/Editor/SubobjectDataInterface/Private/SubobjectData.cpp b/Engine/Source/Editor/SubobjectDataInterface/Private/SubobjectData.cpp
--- a/Engine/Source/Editor/SubobjectDataInterface/Private/SubobjectData.cpp
+++ b/Engine/Source/Editor/SubobjectDataInterface/Private/SubobjectData.cpp
@@ -310,6 +310,11 @@ FString FSubobjectData::GetDisplayString(bool bShowNativeComponentNames /* = tru
 	}
 }
 
+FText FSubobjectData::GetDragDropDisplayText() const
+{
+	return FText::FromString(GetDisplayString());
+}
+
 FText FSubobjectData::GetDisplayNameContextModifiers(bool bShowNativeComponentNames) const
 {
 	if(IsActor())
diff --git a/Engine/Source/Editor/SubobjectDataInterface/Public/SubobjectData.h b/Engine/Source/Editor/SubobjectDataInterface/Public/SubobjectData.h
--- a/Engine/Source/Editor/SubobjectDataInterface/Public/SubobjectData.h
+++ b/Engine/Source/Editor/SubobjectDataInterface/Public/SubobjectData.h
@@ -138,6 +138,9 @@ public:
 
 	virtual FString GetDisplayString(bool bShowNativeComponentNames = true) const;
 
+	/** Get the display name of this subobject that should be used during drag/drop operations */
+	FText GetDragDropDisplayText() const;
+
 	/**
 	 * Returns a string that gives details about this subobject. For Inherited
 	 * components it will return "(Inherited)". For native components
diff --git a/Engine/Source/Editor/SubobjectEditor/Private/SSubobjectEditor.cpp b/Engine/Source/Editor/SubobjectEditor/Private/SSubobjectEditor.cpp
--- a/Engine/Source/Editor/SubobjectEditor/Private/SSubobjectEditor.cpp
+++ b/Engine/Source/Editor/SubobjectEditor/Private/SSubobjectEditor.cpp
@@ -1009,11 +1009,11 @@ void SSubobject_RowWidget::HandleOnDragEnter(const FDragDropEvent& DragDropEvent
 					// Attempted to drag and drop onto self
 					if (DragRowOp->SourceNodes.Num() > 1)
 					{
-						Message = FText::Format(LOCTEXT("DropActionToolTip_Error_CannotAttachToSelfWithMultipleSelection", "Cannot attach the selected components here because it would result in {0} being attached to itself. Remove it from the selection and try again."), DraggedNodeData->GetDisplayName());
+						Message = FText::Format(LOCTEXT("DropActionToolTip_Error_CannotAttachToSelfWithMultipleSelection", "Cannot attach the selected components here because it would result in {0} being attached to itself. Remove it from the selection and try again."), DraggedNodeData->GetDragDropDisplayText());
 					}
 					else
 					{
-						Message = FText::Format(LOCTEXT("DropActionToolTip_Error_CannotAttachToSelf", "Cannot attach {0} to itself."), DraggedNodeData->GetDisplayName());
+						Message = FText::Format(LOCTEXT("DropActionToolTip_Error_CannotAttachToSelf", "Cannot attach {0} to itself."), DraggedNodeData->GetDragDropDisplayText());
 					}
 				}
 				else if (NodePtr->IsAttachedTo(DraggedNodePtr))
@@ -1021,11 +1021,11 @@ void SSubobject_RowWidget::HandleOnDragEnter(const FDragDropEvent& DragDropEvent
 					// Attempted to drop a parent onto a child
 					if (DragRowOp->SourceNodes.Num() > 1)
 					{
-						Message = FText::Format(LOCTEXT("DropActionToolTip_Error_CannotAttachToChildWithMultipleSelection", "Cannot attach the selected components here because it would result in {0} being attached to one of its children. Remove it from the selection and try again."), DraggedNodeData->GetDisplayName());
+						Message = FText::Format(LOCTEXT("DropActionToolTip_Error_CannotAttachToChildWithMultipleSelection", "Cannot attach the selected components here because it would result in {0} being attached to one of its children. Remove it from the selection and try again."), DraggedNodeData->GetDragDropDisplayText());
 					}
 					else
 					{
-						Message = FText::Format(LOCTEXT("DropActionToolTip_Error_CannotAttachToChild", "Cannot attach {0} to one of its children."), DraggedNodeData->GetDisplayName());
+						Message = FText::Format(LOCTEXT("DropActionToolTip_Error_CannotAttachToChild", "Cannot attach {0} to one of its children."), DraggedNodeData->GetDragDropDisplayText());
 					}
 				}
 				else if (HoveredTemplate == nullptr || DraggedTemplate == nullptr)
@@ -1087,12 +1087,12 @@ void SSubobject_RowWidget::HandleOnDragEnter(const FDragDropEvent& DragDropEvent
 							{
 								// Only available action is to copy the dragged node to the other Blueprint and make it the new root
 								// Default root will be deleted
-								Message = FText::Format(LOCTEXT("DropActionToolTip_DropMakeNewRootNodeFromCopyAndDelete", "Drop here to copy {0} to a new variable and make it the new root. The default root will be deleted."), DraggedNodeData->GetDisplayName());
+								Message = FText::Format(LOCTEXT("DropActionToolTip_DropMakeNewRootNodeFromCopyAndDelete", "Drop here to copy {0} to a new variable and make it the new root. The default root will be deleted."), DraggedNodeData->GetDragDropDisplayText());
 							}
 							else
 							{
 								// Only available action is to copy the dragged node to the other Blueprint and make it the new root
-								Message = FText::Format(LOCTEXT("DropActionToolTip_DropMakeNewRootNodeFromCopy", "Drop here to copy {0} to a new variable and make it the new root."), DraggedNodeData->GetDisplayName());
+								Message = FText::Format(LOCTEXT("DropActionToolTip_DropMakeNewRootNodeFromCopy", "Drop here to copy {0} to a new variable and make it the new root."), DraggedNodeData->GetDragDropDisplayText());
 							}
 							DragRowOp->PendingDropAction = FSubobjectRowDragDropOp::DropAction_MakeNewRoot;
 						}
@@ -1101,11 +1101,11 @@ void SSubobject_RowWidget::HandleOnDragEnter(const FDragDropEvent& DragDropEvent
 							// Only available action is to copy the dragged node(s) to the other Blueprint and attach it to the root
 							if (DragRowOp->SourceNodes.Num() > 1)
 							{
-								Message = FText::Format(LOCTEXT("DropActionToolTip_AttachComponentsToThisNodeFromCopyWithMultipleSelection", "Drop here to copy the selected components to new variables and attach them to {0}."), NodeData->GetDisplayName());
+								Message = FText::Format(LOCTEXT("DropActionToolTip_AttachComponentsToThisNodeFromCopyWithMultipleSelection", "Drop here to copy the selected components to new variables and attach them to {0}."), NodeData->GetDragDropDisplayText());
 							}
 							else
 							{
-								Message = FText::Format(LOCTEXT("DropActionToolTip_AttachToThisNodeFromCopy", "Drop here to copy {0} to a new variable and attach it to {1}."), DraggedNodeData->GetDisplayName(), NodeData->GetDisplayName());
+								Message = FText::Format(LOCTEXT("DropActionToolTip_AttachToThisNodeFromCopy", "Drop here to copy {0} to a new variable and attach it to {1}."), DraggedNodeData->GetDragDropDisplayText(), NodeData->GetDragDropDisplayText());
 							}
 
 							DragRowOp->PendingDropAction = FSubobjectRowDragDropOp::DropAction_AttachTo;
@@ -1117,12 +1117,12 @@ void SSubobject_RowWidget::HandleOnDragEnter(const FDragDropEvent& DragDropEvent
 						{
 							// Only available action is to make the dragged node the new root
 							// Default root will be deleted
-							Message = FText::Format(LOCTEXT("DropActionToolTip_DropMakeNewRootNodeAndDelete", "Drop here to make {0} the new root. The default root will be deleted."), DraggedNodeData->GetDisplayName());
+							Message = FText::Format(LOCTEXT("DropActionToolTip_DropMakeNewRootNodeAndDelete", "Drop here to make {0} the new root. The default root will be deleted."), DraggedNodeData->GetDragDropDisplayText());
 						}
 						else
 						{
 							// Only available action is to make the dragged node the new root
-							Message = FText::Format(LOCTEXT("DropActionToolTip_DropMakeNewRootNode", "Drop here to make {0} the new root."), DraggedNodeData->GetDisplayName());
+							Message = FText::Format(LOCTEXT("DropActionToolTip_DropMakeNewRootNode", "Drop here to make {0} the new root."), DraggedNodeData->GetDragDropDisplayText());
 						}
 						DragRowOp->PendingDropAction = FSubobjectRowDragDropOp::DropAction_MakeNewRoot;
 					}
@@ -1131,11 +1131,11 @@ void SSubobject_RowWidget::HandleOnDragEnter(const FDragDropEvent& DragDropEvent
 						// Only available action is to attach the dragged node(s) to the root
 						if (DragRowOp->SourceNodes.Num() > 1)
 						{
-							Message = FText::Format(LOCTEXT("DropActionToolTip_AttachToThisNodeWithMultipleSelection", "Drop here to attach the selected components to {0}."), NodeData->GetDisplayName());
+							Message = FText::Format(LOCTEXT("DropActionToolTip_AttachToThisNodeWithMultipleSelection", "Drop here to attach the selected components to {0}."), NodeData->GetDragDropDisplayText());
 						}
 						else
 						{
-							Message = FText::Format(LOCTEXT("DropActionToolTip_AttachToThisNode", "Drop here to attach {0} to {1}."), DraggedNodeData->GetDisplayName(), NodeData->GetDisplayName());
+							Message = FText::Format(LOCTEXT("DropActionToolTip_AttachToThisNode", "Drop here to attach {0} to {1}."), DraggedNodeData->GetDragDropDisplayText(), NodeData->GetDragDropDisplayText());
 						}
 
 						DragRowOp->PendingDropAction = FSubobjectRowDragDropOp::DropAction_AttachTo;
@@ -1146,11 +1146,11 @@ void SSubobject_RowWidget::HandleOnDragEnter(const FDragDropEvent& DragDropEvent
 					// Detach the dropped node(s) from the current node and reattach to the root node
 					if (DragRowOp->SourceNodes.Num() > 1)
 					{
-						Message = FText::Format(LOCTEXT("DropActionToolTip_DetachFromThisNodeWithMultipleSelection", "Drop here to detach the selected components from {0}."), NodeData->GetDisplayName());
+						Message = FText::Format(LOCTEXT("DropActionToolTip_DetachFromThisNodeWithMultipleSelection", "Drop here to detach the selected components from {0}."), NodeData->GetDragDropDisplayText());
 					}
 					else
 					{
-						Message = FText::Format(LOCTEXT("DropActionToolTip_DetachFromThisNode", "Drop here to detach {0} from {1}."), DraggedNodeData->GetDisplayName(), NodeData->GetDisplayName());
+						Message = FText::Format(LOCTEXT("DropActionToolTip_DetachFromThisNode", "Drop here to detach {0} from {1}."), DraggedNodeData->GetDragDropDisplayText(), NodeData->GetDragDropDisplayText());
 					}
 
 					DragRowOp->PendingDropAction = FSubobjectRowDragDropOp::DropAction_DetachFrom;
@@ -1187,20 +1187,20 @@ void SSubobject_RowWidget::HandleOnDragEnter(const FDragDropEvent& DragDropEvent
 					{
 						if (DragRowOp->SourceNodes.Num() > 1)
 						{
-							Message = FText::Format(LOCTEXT("DropActionToolTip_AttachToThisNodeFromCopyWithMultipleSelection", "Drop here to copy the selected nodes to new variables and attach them to {0}."), NodeData->GetDisplayName());
+							Message = FText::Format(LOCTEXT("DropActionToolTip_AttachToThisNodeFromCopyWithMultipleSelection", "Drop here to copy the selected nodes to new variables and attach them to {0}."), NodeData->GetDragDropDisplayText());
 						}
 						else
 						{
-							Message = FText::Format(LOCTEXT("DropActionToolTip_AttachToThisNodeFromCopy", "Drop here to copy {0} to a new variable and attach it to {1}."), DraggedNodeData->GetDisplayName(), NodeData->GetDisplayName());
+							Message = FText::Format(LOCTEXT("DropActionToolTip_AttachToThisNodeFromCopy", "Drop here to copy {0} to a new variable and attach it to {1}."), DraggedNodeData->GetDragDropDisplayText(), NodeData->GetDragDropDisplayText());
 						}
 					}
 					else if (DragRowOp->SourceNodes.Num() > 1)
 					{
-						Message = FText::Format(LOCTEXT("DropActionToolTip_AttachToThisNodeWithMultipleSelection", "Drop here to attach the selected components to {0}."), NodeData->GetDisplayName());
+						Message = FText::Format(LOCTEXT("DropActionToolTip_AttachToThisNodeWithMultipleSelection", "Drop here to attach the selected components to {0}."), NodeData->GetDragDropDisplayText());
 					}
 					else
 					{
-						Message = FText::Format(LOCTEXT("DropActionToolTip_AttachToThisNode", "Drop here to attach {0} to {1}."), DraggedNodeData->GetDisplayName(), NodeData->GetDisplayName());
+						Message = FText::Format(LOCTEXT("DropActionToolTip_AttachToThisNode", "Drop here to attach {0} to {1}."), DraggedNodeData->GetDragDropDisplayText(), NodeData->GetDragDropDisplayText());
 					}
 
 					DragRowOp->PendingDropAction = FSubobjectRowDragDropOp::DropAction_AttachTo;
@@ -1208,7 +1208,7 @@ void SSubobject_RowWidget::HandleOnDragEnter(const FDragDropEvent& DragDropEvent
 				else
 				{
 					// The dropped node cannot be attached to the current node
-					Message = FText::Format(LOCTEXT("DropActionToolTip_Error_TooManyAttachments", "Unable to attach {0} to {1}."), DraggedNodeData->GetDisplayName(), NodeData->GetDisplayName());
+					Message = FText::Format(LOCTEXT("DropActionToolTip_Error_TooManyAttachments", "Unable to attach {0} to {1}."), DraggedNodeData->GetDragDropDisplayText(), NodeData->GetDragDropDisplayText());
 				}
 			}
 		}
