Automatically select a new component for rename upon duplication.

#rnx

diff --git a/Engine/Source/Editor/SubobjectDataInterface/Private/SubobjectDataSubsystem.cpp b/Engine/Source/Editor/SubobjectDataInterface/Private/SubobjectDataSubsystem.cpp
--- a/Engine/Source/Editor/SubobjectDataInterface/Private/SubobjectDataSubsystem.cpp
+++ b/Engine/Source/Editor/SubobjectDataInterface/Private/SubobjectDataSubsystem.cpp
@@ -1998,7 +1998,7 @@ void USubobjectDataSubsystem::PasteSubobjects(const FSubobjectDataHandle& PasteT
 	}
 }
 
-void USubobjectDataSubsystem::DuplicateSubobjects(const FSubobjectDataHandle& Context, const TArray<FSubobjectDataHandle>& SubobjectsToDup, UBlueprint* BpContext)
+void USubobjectDataSubsystem::DuplicateSubobjects(const FSubobjectDataHandle& Context, const TArray<FSubobjectDataHandle>& SubobjectsToDup, UBlueprint* BpContext, TArray<FSubobjectDataHandle>& OutNewSubobjects)
 {
 	if(!Context.IsValid() || SubobjectsToDup.IsEmpty())
 	{
@@ -2046,6 +2046,11 @@ void USubobjectDataSubsystem::DuplicateSubobjects(const FSubobjectDataHandle& Co
 			}
 			
 			FSubobjectDataHandle ClonedSubobject = AddNewSubobject(NewSubobjectParams, FailedAddReason);
+			if (ClonedSubobject.IsValid())
+			{
+				OutNewSubobjects.Add(ClonedSubobject);
+			}
+
 			FSubobjectData* ClonedData = ClonedSubobject.GetSharedDataPtr().Get();
 			
 			if(ClonedData && ClonedData->IsSceneComponent())
diff --git a/Engine/Source/Editor/SubobjectDataInterface/Public/SubobjectDataSubsystem.h b/Engine/Source/Editor/SubobjectDataInterface/Public/SubobjectDataSubsystem.h
--- a/Engine/Source/Editor/SubobjectDataInterface/Public/SubobjectDataSubsystem.h
+++ b/Engine/Source/Editor/SubobjectDataInterface/Public/SubobjectDataSubsystem.h
@@ -332,8 +332,9 @@ public:
 	 * @param Context			The owning context that the subobjects to dup come from
 	 * @param SubobjectsToDup	Array of handles of existing subobejcts you would like to have duplicated
 	 * @param BpContext			Pointer to the current blueprint context if necessary. Use nullptr if dealing with instances
+	 * @param OutNewSubobjects	Array that will be populated with any newly created subobjects
 	 */
-	void DuplicateSubobjects(const FSubobjectDataHandle& Context, const TArray<FSubobjectDataHandle>& SubobjectsToDup, UBlueprint* BpContext);
+	void DuplicateSubobjects(const FSubobjectDataHandle& Context, const TArray<FSubobjectDataHandle>& SubobjectsToDup, UBlueprint* BpContext, TArray<FSubobjectDataHandle>& OutNewSubobjects);
 	
 	FScopedTransaction* BeginTransaction(const TArray<FSubobjectDataHandle>& Handles, const FText& Description, UBlueprint* Blueprint);
 
diff --git a/Engine/Source/Editor/SubobjectEditor/Private/SSubobjectBlueprintEditor.cpp b/Engine/Source/Editor/SubobjectEditor/Private/SSubobjectBlueprintEditor.cpp
--- a/Engine/Source/Editor/SubobjectEditor/Private/SSubobjectBlueprintEditor.cpp
+++ b/Engine/Source/Editor/SubobjectEditor/Private/SSubobjectBlueprintEditor.cpp
@@ -288,12 +288,21 @@ void SSubobjectBlueprintEditor::OnDuplicateComponent()
 		// this would prevent the completion of the 'rename' or 'create + give initial name' transaction (occurring on focus lost).
 		FSlateApplication::Get().ClearKeyboardFocus();
 		
-		const FScopedTransaction Transaction(SelectedNodes.Num() > 1 ? LOCTEXT("DuplicateComponents", "Duplicate Components") : LOCTEXT("DuplicateComponent", "Duplicate Component"));
+		TUniquePtr<FScopedTransaction> Transaction = MakeUnique<FScopedTransaction>(SelectedNodes.Num() > 1 ? LOCTEXT("DuplicateComponents", "Duplicate Components") : LOCTEXT("DuplicateComponent", "Duplicate Component"));
 		
 		if(USubobjectDataSubsystem* System = USubobjectDataSubsystem::Get())
 		{
-			System->DuplicateSubobjects(GetObjectContextHandle(), SelectedNodes, GetBlueprint());
+			TArray<FSubobjectDataHandle> DuplicatedHandles;
+			System->DuplicateSubobjects(GetObjectContextHandle(), SelectedNodes, GetBlueprint(), DuplicatedHandles);
 			UpdateTree();
+
+			// Set focus to the newly created subobject
+			FSubobjectEditorTreeNodePtrType NewNode = DuplicatedHandles.Num() > 0 ? FindSlateNodeForHandle(DuplicatedHandles[0]) : nullptr;
+			if (NewNode != nullptr)
+			{
+				TreeWidget->SetSelection(NewNode);
+				OnRenameComponent(MoveTemp(Transaction));
+			}
 		}
 	}
 }
diff --git a/Engine/Source/Editor/SubobjectEditor/Private/SSubobjectInstanceEditor.cpp b/Engine/Source/Editor/SubobjectEditor/Private/SSubobjectInstanceEditor.cpp
--- a/Engine/Source/Editor/SubobjectEditor/Private/SSubobjectInstanceEditor.cpp
+++ b/Engine/Source/Editor/SubobjectEditor/Private/SSubobjectInstanceEditor.cpp
@@ -227,11 +227,21 @@ void SSubobjectInstanceEditor::OnDuplicateComponent()
 		// this would prevent the completion of the 'rename' or 'create + give initial name' transaction (occurring on focus lost).
 		FSlateApplication::Get().ClearKeyboardFocus();
 		
-		const FScopedTransaction Transaction(SelectedNodes.Num() > 1 ? LOCTEXT("DuplicateComponents", "Duplicate Components") : LOCTEXT("DuplicateComponent", "Duplicate Component"));
+		TUniquePtr<FScopedTransaction> Transaction = MakeUnique<FScopedTransaction>(SelectedNodes.Num() > 1 ? LOCTEXT("DuplicateComponents", "Duplicate Components") : LOCTEXT("DuplicateComponent", "Duplicate Component"));
+
 		if(USubobjectDataSubsystem* System = USubobjectDataSubsystem::Get())
 		{
-			System->DuplicateSubobjects(GetObjectContextHandle(), SelectedNodes, /* BpContext = */ nullptr);
+			TArray<FSubobjectDataHandle> DuplicatedHandles;
+			System->DuplicateSubobjects(GetObjectContextHandle(), SelectedNodes, /* BpContext = */ nullptr, DuplicatedHandles);
 			UpdateTree();
+
+			// Set focus to the newly created subobject
+			FSubobjectEditorTreeNodePtrType NewNode = DuplicatedHandles.Num() > 0 ? FindSlateNodeForHandle(DuplicatedHandles[0]) : nullptr;
+			if (NewNode != nullptr)
+			{
+				TreeWidget->SetSelection(NewNode);
+				OnRenameComponent(MoveTemp(Transaction));
+			}
 		}
 	}
 }
