Add Apply modifications to blueprints to the context menu in the subobject instance editor.
Make SSubobjectEditor::CreateMenuBuilder non-const to allow for the binding of functions that can modify the actor instance like reseting to blueprint defaults

#rnx

diff --git a/Engine/Source/Editor/SubobjectEditor/Private/SSubobjectEditor.cpp b/Engine/Source/Editor/SubobjectEditor/Private/SSubobjectEditor.cpp
--- a/Engine/Source/Editor/SubobjectEditor/Private/SSubobjectEditor.cpp
+++ b/Engine/Source/Editor/SubobjectEditor/Private/SSubobjectEditor.cpp
@@ -1925,7 +1925,7 @@ bool SSubobject_RowWidget::IsReadOnly() const
 ////////////////////////////////////////////////
 // SSubobjectEditor
 
-FMenuBuilder SSubobjectEditor::CreateMenuBuilder() const
+FMenuBuilder SSubobjectEditor::CreateMenuBuilder()
 {
 	// Menu builder for editing a blueprint
 	FMenuBuilder EditBlueprintMenuBuilder(true, nullptr);
diff --git a/Engine/Source/Editor/SubobjectEditor/Private/SSubobjectInstanceEditor.cpp b/Engine/Source/Editor/SubobjectEditor/Private/SSubobjectInstanceEditor.cpp
--- a/Engine/Source/Editor/SubobjectEditor/Private/SSubobjectInstanceEditor.cpp
+++ b/Engine/Source/Editor/SubobjectEditor/Private/SSubobjectInstanceEditor.cpp
@@ -12,12 +12,15 @@
 #include "Styling/SlateIconFinder.h"
 #include "SlateOptMacros.h"
 #include "Widgets/Input/SSearchBox.h"
+#include "Widgets/Notifications/SNotificationList.h"
+#include "Framework/Notifications/NotificationManager.h"
 #include "Framework/MultiBox/MultiBoxBuilder.h"
 #include "SComponentClassCombo.h"
 #include "SEditorHeaderButton.h"
 #include "Editor/UnrealEdEngine.h"
 #include "Subsystems/PanelExtensionSubsystem.h"	// SExtensionPanel
 #include "Kismet2/ComponentEditorUtils.h"
+#include "Kismet2/KismetEditorUtilities.h"		// ApplyInstanceChangesToBlueprint
 
 #define LOCTEXT_NAMESPACE "SSubobjectInstanceEditor"
 
@@ -430,6 +433,42 @@ void SSubobjectInstanceEditor::PopulateContextMenuImpl(UToolMenu* InMenu, TArray
     FComponentEditorUtils::FillComponentContextMenuOptions(InMenu, SelectedComponents);
 }
 
+FMenuBuilder SSubobjectInstanceEditor::CreateMenuBuilder()
+{
+	FMenuBuilder EditBlueprintMenuBuilder = SSubobjectEditor::CreateMenuBuilder();
+
+	EditBlueprintMenuBuilder.BeginSection(NAME_None, LOCTEXT("EditBlueprintMenu_InstanceHeader", "Instance modifications"));
+
+	EditBlueprintMenuBuilder.AddMenuEntry
+	(
+		LOCTEXT("PushChangesToBlueprint", "Apply Instance Changes to Blueprint"),
+		TAttribute<FText>(this, &SSubobjectInstanceEditor::OnGetApplyChangesToBlueprintTooltip),
+		FSlateIcon(),
+		FUIAction(FExecuteAction::CreateSP(this, &SSubobjectInstanceEditor::OnApplyChangesToBlueprint))
+	);
+
+	EditBlueprintMenuBuilder.AddMenuEntry
+	(
+		LOCTEXT("ResetToDefault", "Reset Instance Changes to Blueprint Default"),
+		TAttribute<FText>(this, &SSubobjectInstanceEditor::OnGetResetToBlueprintDefaultsTooltip),
+		FSlateIcon(),
+		FUIAction(FExecuteAction::CreateSP(this, &SSubobjectInstanceEditor::OnResetToBlueprintDefaults))
+	);
+
+	EditBlueprintMenuBuilder.BeginSection(NAME_None, LOCTEXT("EditBlueprintMenu_NewHeader", "Create New"));
+	//EditBlueprintMenuBuilder.AddMenuSeparator();
+
+	EditBlueprintMenuBuilder.AddMenuEntry
+	(
+		LOCTEXT("CreateChildBlueprint", "Create Child Blueprint Class"),
+		LOCTEXT("CreateChildBlueprintTooltip", "Creates a Child Blueprint Class based on the current Blueprint, allowing you to create variants easily.  This replaces the current actor instance with a new one based on the new Child Blueprint Class."),
+		FSlateIcon(),
+		FUIAction(FExecuteAction::CreateSP(this, &SSubobjectInstanceEditor::PromoteToBlueprint))
+	);
+
+	return EditBlueprintMenuBuilder;
+}
+
 FSlateColor SSubobjectInstanceEditor::GetColorTintForIcon(FSubobjectEditorTreeNodePtrType Node) const
 {
 	static const FLinearColor IntroducedHereColor(FLinearColor::White);
@@ -475,4 +514,175 @@ FSlateColor SSubobjectInstanceEditor::GetColorTintForIcon(FSubobjectEditorTreeNo
 	return IntroducedHereColor;
 }
 
+FText SSubobjectInstanceEditor::OnGetApplyChangesToBlueprintTooltip() const
+{
+	int32 NumChangedProperties = 0;
+
+	AActor* Actor = Cast<AActor>(GetObjectContext());
+	UBlueprint* Blueprint = (Actor != nullptr) ? Cast<UBlueprint>(Actor->GetClass()->ClassGeneratedBy) : nullptr;
+
+	if (Actor != nullptr && Blueprint != nullptr && Actor->GetClass()->ClassGeneratedBy == Blueprint)
+	{
+		AActor* BlueprintCDO = Actor->GetClass()->GetDefaultObject<AActor>();
+		if (BlueprintCDO != nullptr)
+		{
+			const EditorUtilities::ECopyOptions::Type CopyOptions = (EditorUtilities::ECopyOptions::Type)(EditorUtilities::ECopyOptions::PreviewOnly | EditorUtilities::ECopyOptions::OnlyCopyEditOrInterpProperties | EditorUtilities::ECopyOptions::SkipInstanceOnlyProperties);
+			NumChangedProperties += EditorUtilities::CopyActorProperties(Actor, BlueprintCDO, CopyOptions);
+		}
+		NumChangedProperties += Actor->GetInstanceComponents().Num();
+	}
+
+	if (NumChangedProperties == 0)
+	{
+		return LOCTEXT("DisabledPushToBlueprintDefaults_ToolTip", "Replaces the Blueprint's defaults with any altered property values.");
+	}
+	else if (NumChangedProperties > 1)
+	{
+		return FText::Format(LOCTEXT("PushToBlueprintDefaults_ToolTip", "Click to apply {0} changed properties to the Blueprint."), FText::AsNumber(NumChangedProperties));
+	}
+	else
+	{
+		return LOCTEXT("PushOneToBlueprintDefaults_ToolTip", "Click to apply 1 changed property to the Blueprint.");
+	}
+}
+
+void SSubobjectInstanceEditor::OnApplyChangesToBlueprint() const
+{
+	AActor* Actor = Cast<AActor>(GetObjectContext());
+	const UBlueprint* const Blueprint = (Actor != nullptr) ? Cast<UBlueprint>(Actor->GetClass()->ClassGeneratedBy) : nullptr;
+
+	if (Actor != nullptr && Blueprint != nullptr && Actor->GetClass()->ClassGeneratedBy == Blueprint)
+	{
+		const FString ActorLabel = Actor->GetActorLabel();
+		int32 NumChangedProperties = FKismetEditorUtilities::ApplyInstanceChangesToBlueprint(Actor);
+
+		// Set up a notification record to indicate success/failure
+		FNotificationInfo NotificationInfo(FText::GetEmpty());
+		NotificationInfo.FadeInDuration = 1.0f;
+		NotificationInfo.FadeOutDuration = 2.0f;
+		NotificationInfo.bUseLargeFont = false;
+		SNotificationItem::ECompletionState CompletionState;
+		if (NumChangedProperties > 0)
+		{
+			if (NumChangedProperties > 1)
+			{
+				FFormatNamedArguments Args;
+				Args.Add(TEXT("BlueprintName"), FText::FromName(Blueprint->GetFName()));
+				Args.Add(TEXT("NumChangedProperties"), NumChangedProperties);
+				Args.Add(TEXT("ActorName"), FText::FromString(ActorLabel));
+				NotificationInfo.Text = FText::Format(LOCTEXT("PushToBlueprintDefaults_ApplySuccess", "Updated Blueprint {BlueprintName} ({NumChangedProperties} property changes applied from actor {ActorName})."), Args);
+			}
+			else
+			{
+				FFormatNamedArguments Args;
+				Args.Add(TEXT("BlueprintName"), FText::FromName(Blueprint->GetFName()));
+				Args.Add(TEXT("ActorName"), FText::FromString(ActorLabel));
+				NotificationInfo.Text = FText::Format(LOCTEXT("PushOneToBlueprintDefaults_ApplySuccess", "Updated Blueprint {BlueprintName} (1 property change applied from actor {ActorName})."), Args);
+			}
+			CompletionState = SNotificationItem::CS_Success;
+		}
+		else
+		{
+			NotificationInfo.Text = LOCTEXT("PushToBlueprintDefaults_ApplyFailed", "No properties were copied");
+			CompletionState = SNotificationItem::CS_Fail;
+		}
+
+		// Add the notification to the queue
+		const TSharedPtr<SNotificationItem> Notification = FSlateNotificationManager::Get().AddNotification(NotificationInfo);
+		Notification->SetCompletionState(CompletionState);
+	}
+}
+
+void SSubobjectInstanceEditor::OnResetToBlueprintDefaults()
+{
+	int32 NumChangedProperties = 0;
+
+	AActor* Actor = Cast<AActor>(GetObjectContext());
+	UBlueprint* Blueprint = (Actor != nullptr) ? Cast<UBlueprint>(Actor->GetClass()->ClassGeneratedBy) : nullptr;
+
+	if ((Actor != nullptr) && (Blueprint != nullptr) && (Actor->GetClass()->ClassGeneratedBy == Blueprint))
+	{
+		const FScopedTransaction Transaction(LOCTEXT("ResetToBlueprintDefaults_Transaction", "Reset to Class Defaults"));
+
+		{
+			AActor* BlueprintCDO = Actor->GetClass()->GetDefaultObject<AActor>();
+			if (BlueprintCDO != nullptr)
+			{
+				const EditorUtilities::ECopyOptions::Type CopyOptions = (EditorUtilities::ECopyOptions::Type)(EditorUtilities::ECopyOptions::OnlyCopyEditOrInterpProperties);
+				NumChangedProperties = EditorUtilities::CopyActorProperties(BlueprintCDO, Actor, CopyOptions);
+			}
+			NumChangedProperties += Actor->GetInstanceComponents().Num();
+			Actor->ClearInstanceComponents(true);
+		}
+
+		// Set up a notification record to indicate success/failure
+		FNotificationInfo NotificationInfo(FText::GetEmpty());
+		NotificationInfo.FadeInDuration = 1.0f;
+		NotificationInfo.FadeOutDuration = 2.0f;
+		NotificationInfo.bUseLargeFont = false;
+		SNotificationItem::ECompletionState CompletionState;
+		if (NumChangedProperties > 0)
+		{
+			if (NumChangedProperties > 1)
+			{
+				FFormatNamedArguments Args;
+				Args.Add(TEXT("BlueprintName"), FText::FromName(Blueprint->GetFName()));
+				Args.Add(TEXT("NumChangedProperties"), NumChangedProperties);
+				Args.Add(TEXT("ActorName"), FText::FromString(Actor->GetActorLabel()));
+				NotificationInfo.Text = FText::Format(LOCTEXT("ResetToBlueprintDefaults_ApplySuccess", "Reset {ActorName} ({NumChangedProperties} property changes applied from Blueprint {BlueprintName})."), Args);
+			}
+			else
+			{
+				FFormatNamedArguments Args;
+				Args.Add(TEXT("BlueprintName"), FText::FromName(Blueprint->GetFName()));
+				Args.Add(TEXT("ActorName"), FText::FromString(Actor->GetActorLabel()));
+				NotificationInfo.Text = FText::Format(LOCTEXT("ResetOneToBlueprintDefaults_ApplySuccess", "Reset {ActorName} (1 property change applied from Blueprint {BlueprintName})."), Args);
+			}
+			CompletionState = SNotificationItem::CS_Success;
+		}
+		else
+		{
+			NotificationInfo.Text = LOCTEXT("ResetToBlueprintDefaults_Failed", "No properties were reset");
+			CompletionState = SNotificationItem::CS_Fail;
+		}
+
+		UpdateTree();
+
+		// Add the notification to the queue
+		const TSharedPtr<SNotificationItem> Notification = FSlateNotificationManager::Get().AddNotification(NotificationInfo);
+		Notification->SetCompletionState(CompletionState);
+	}
+}
+
+FText SSubobjectInstanceEditor::OnGetResetToBlueprintDefaultsTooltip() const
+{
+	int32 NumChangedProperties = 0;
+
+	AActor* Actor = Cast<AActor>(GetObjectContext());
+	UBlueprint* Blueprint = (Actor != nullptr) ? Cast<UBlueprint>(Actor->GetClass()->ClassGeneratedBy) : nullptr;
+	if (Actor != nullptr && Blueprint != nullptr && Actor->GetClass()->ClassGeneratedBy == Blueprint)
+	{
+		AActor* BlueprintCDO = Actor->GetClass()->GetDefaultObject<AActor>();
+		if (BlueprintCDO != nullptr)
+		{
+			const EditorUtilities::ECopyOptions::Type CopyOptions = (EditorUtilities::ECopyOptions::Type)(EditorUtilities::ECopyOptions::PreviewOnly | EditorUtilities::ECopyOptions::OnlyCopyEditOrInterpProperties);
+			NumChangedProperties += EditorUtilities::CopyActorProperties(BlueprintCDO, Actor, CopyOptions);
+		}
+		NumChangedProperties += Actor->GetInstanceComponents().Num();
+	}
+
+	if (NumChangedProperties == 0)
+	{
+		return LOCTEXT("DisabledResetBlueprintDefaults_ToolTip", "Resets altered properties back to their Blueprint default values.");
+	}
+	else if (NumChangedProperties > 1)
+	{
+		return FText::Format(LOCTEXT("ResetToBlueprintDefaults_ToolTip", "Click to reset {0} changed properties to their Blueprint default values."), FText::AsNumber(NumChangedProperties));
+	}
+	else
+	{
+		return LOCTEXT("ResetOneToBlueprintDefaults_ToolTip", "Click to reset 1 changed property to its Blueprint default value.");
+	}
+}
+
 #undef LOCTEXT_NAMESPACE
diff --git a/Engine/Source/Editor/SubobjectEditor/Public/SSubobjectEditor.h b/Engine/Source/Editor/SubobjectEditor/Public/SSubobjectEditor.h
--- a/Engine/Source/Editor/SubobjectEditor/Public/SSubobjectEditor.h
+++ b/Engine/Source/Editor/SubobjectEditor/Public/SSubobjectEditor.h
@@ -627,7 +627,7 @@ protected:
 	 */
 	FSubobjectDataHandle CachedRootHandle;
 
-	virtual FMenuBuilder CreateMenuBuilder() const;
+	virtual FMenuBuilder CreateMenuBuilder();
 
 	/** Constructs the slate drag/drop tree for this subobject editor */
 	virtual void ConstructTreeWidget();
diff --git a/Engine/Source/Editor/SubobjectEditor/Public/SSubobjectInstanceEditor.h b/Engine/Source/Editor/SubobjectEditor/Public/SSubobjectInstanceEditor.h
--- a/Engine/Source/Editor/SubobjectEditor/Public/SSubobjectInstanceEditor.h
+++ b/Engine/Source/Editor/SubobjectEditor/Public/SSubobjectInstanceEditor.h
@@ -44,6 +44,7 @@ protected:
     virtual TSharedPtr<SWidget> BuildSceneRootDropActionMenu(FSubobjectEditorTreeNodePtrType DroppedOntoNodePtr, FSubobjectEditorTreeNodePtrType DroppedNodePtr) override;
 	virtual FSubobjectDataHandle AddNewSubobject(const FSubobjectDataHandle& ParentHandle, UClass* NewClass, UObject* AssetOverride, FText& OutFailReason, TUniquePtr<FScopedTransaction> InOngoingTransaction) override;
 	virtual void PopulateContextMenuImpl(UToolMenu* InMenu, TArray<FSubobjectEditorTreeNodePtrType>& InSelectedItems, bool bIsChildActorSubtreeNodeSelected) override;
+	virtual FMenuBuilder CreateMenuBuilder();
 
 public:
 	virtual FSlateColor GetColorTintForIcon(FSubobjectEditorTreeNodePtrType Node) const override;
@@ -52,4 +53,18 @@ public:
 	/** Delegate to invoke when objects within the Subobject tree are replaced (eg, via re-instancing from a BP compile) */
 	FSimpleDelegate OnObjectReplaced;
 
+private:
+
+	/** @return the tooltip describing how many properties will be applied to the blueprint */
+	FText OnGetApplyChangesToBlueprintTooltip() const;
+
+	/** Propagates instance changes to the blueprint */
+	void OnApplyChangesToBlueprint() const;
+
+	/** Resets instance changes to the blueprint default */
+	void OnResetToBlueprintDefaults();
+
+	/** @return the tooltip describing how many properties will be reset to the blueprint default*/
+	FText OnGetResetToBlueprintDefaultsTooltip() const;
+
 };
\ No newline at end of file
