Select the root if no nodes are selected when pasting in the subobject instance editor

#rnx

diff --git a/Engine/Source/Editor/SubobjectEditor/Private/SSubobjectInstanceEditor.cpp b/Engine/Source/Editor/SubobjectEditor/Private/SSubobjectInstanceEditor.cpp
--- a/Engine/Source/Editor/SubobjectEditor/Private/SSubobjectInstanceEditor.cpp
+++ b/Engine/Source/Editor/SubobjectEditor/Private/SSubobjectInstanceEditor.cpp
@@ -254,7 +254,19 @@ void SSubobjectInstanceEditor::PasteNodes()
 	if(USubobjectDataSubsystem* System = USubobjectDataSubsystem::Get())
 	{
 		TArray<FSubobjectDataHandle> OutHandles;
-		System->PasteSubobjects(GetObjectContextHandle(), GetSelectedHandles(), nullptr, OutHandles);
+		TArray<FSubobjectDataHandle> HandlesToPasteOnto = GetSelectedHandles();
+		if (HandlesToPasteOnto.IsEmpty())
+		{
+			if (FSubobjectEditorTreeNodePtrType RootPtr = GetSceneRootNode())
+			{
+				if (RootPtr.IsValid())
+				{
+					HandlesToPasteOnto.Emplace(RootPtr->GetDataHandle());
+				}
+			}
+		}
+
+		System->PasteSubobjects(GetObjectContextHandle(), HandlesToPasteOnto, nullptr, OutHandles);
 
 		if(OutHandles.Num() > 0)
 		{
