Correctly set the parent node when handling the drag/drop of an asset type in the blueprint editor

#rnx

diff --git a/Engine/Source/Editor/SubobjectEditor/Private/SSubobjectEditor.cpp b/Engine/Source/Editor/SubobjectEditor/Private/SSubobjectEditor.cpp
--- a/Engine/Source/Editor/SubobjectEditor/Private/SSubobjectEditor.cpp
+++ b/Engine/Source/Editor/SubobjectEditor/Private/SSubobjectEditor.cpp
@@ -645,6 +645,7 @@ void SSubobjectEditorDragDropTree::Construct(const FArguments& InArgs)
 		.OnEnteredBadState(InArgs._OnTableViewBadState)
 		.HighlightParentNodesForSelection(true);
 
+	SubobjectEditor = InArgs._SubobjectEditor;
 	STreeView<FSubobjectEditorTreeNodePtrType>::Construct(BaseArgs);
 }
 
@@ -2533,7 +2534,15 @@ FReply SSubobjectEditor::TryHandleAssetDragDropOperation(const FDragDropEvent& D
 				USubobjectDataSubsystem* System = USubobjectDataSubsystem::Get();
 				check(System);
 
+				TUniquePtr<FScopedTransaction> AddTransaction = MakeUnique<FScopedTransaction>(LOCTEXT("AddComponent", "Add Component"));
+
 				FAddNewSubobjectParams NewComponentParams;
+
+				// Attach to the currently selected handle, if there are none then select the root object
+				const TArray<FSubobjectDataHandle>& SelectedHandles = GetSelectedHandles();
+				NewComponentParams.ParentHandle = SelectedHandles.IsEmpty() ? GetObjectContextHandle() : SelectedHandles[0];
+				NewComponentParams.BlueprintContext = ShouldModifyBPOnAssetDrop() ? GetBlueprint() : nullptr;
+
 				FText FailReason;
 				NewComponentParams.bSkipMarkBlueprintModified = true;
 				const bool bSetFocusToNewItem = (DroppedAssetIdx == NumAssets - 1);// Only set focus to the last item created
@@ -2558,8 +2567,8 @@ FReply SSubobjectEditor::TryHandleAssetDragDropOperation(const FDragDropEvent& D
 				}
 				else if ((PotentialActorClass != nullptr) && !PotentialActorClass->HasAnyClassFlags(CLASS_Deprecated | CLASS_Abstract | CLASS_NewerVersionExists | CLASS_NotPlaceable))
 				{
-					NewComponentParams.AssetOverride = UChildActorComponent::StaticClass();
-					NewComponentParams.AssetOverride = PotentialActorClass;					
+					NewComponentParams.NewClass = UChildActorComponent::StaticClass();
+					NewComponentParams.AssetOverride = PotentialActorClass;
 					bMarkBlueprintAsModified = true;
 					System->AddNewSubobject(NewComponentParams, FailReason);
 				}
