Also check for the component archetype when updating the tree selection in SActorDetails because it is a viable option for inherited BP components on instanced objects

#rnx

diff --git a/Engine/Source/Editor/LevelEditor/Private/SActorDetails.cpp b/Engine/Source/Editor/LevelEditor/Private/SActorDetails.cpp
--- a/Engine/Source/Editor/LevelEditor/Private/SActorDetails.cpp
+++ b/Engine/Source/Editor/LevelEditor/Private/SActorDetails.cpp
@@ -617,7 +617,7 @@ void SActorDetails::UpdateComponentTreeFromEditorSelection()
 		{
 			TreeWidget->RequestScrollIntoView(TreeNode);
 			TreeWidget->SetItemSelection(TreeNode, true);
-			check(InComponent == TreeNode->GetComponentTemplate());
+			check(InComponent == TreeNode->GetComponentTemplate() || InComponent->GetArchetype() == TreeNode->GetComponentTemplate());
 		}
 		return true;
 	});
diff --git a/Engine/Source/Editor/SubobjectEditor/Private/SSubobjectEditor.cpp b/Engine/Source/Editor/SubobjectEditor/Private/SSubobjectEditor.cpp
--- a/Engine/Source/Editor/SubobjectEditor/Private/SSubobjectEditor.cpp
+++ b/Engine/Source/Editor/SubobjectEditor/Private/SSubobjectEditor.cpp
@@ -2853,19 +2853,29 @@ FSubobjectEditorTreeNodePtrType SSubobjectEditor::FindSlateNodeForObject(const U
 {
 	FSubobjectEditorTreeNodePtrType OutNodePtr;
 
-	if (RootNodes.Num() > 0)
+	if (InObject && RootNodes.Num() > 0)
 	{
 		TSet<FSubobjectEditorTreeNodePtrType> VisitedNodes;
+		UBlueprint* BP = GetBlueprint();
+
 		DepthFirstTraversal(RootNodes[0], VisitedNodes,
-            [&OutNodePtr, &InObject](
+            [&OutNodePtr, InObject, BP](
             const FSubobjectEditorTreeNodePtrType& CurNodePtr)
             {
                 if(CurNodePtr->GetDataHandle().IsValid())
                 {
-                    if(CurNodePtr->GetDataHandle().GetSharedDataPtr()->GetObject() == InObject)
-                    {
-                        OutNodePtr = CurNodePtr;
-                    }
+					if (const FSubobjectData* Data = CurNodePtr->GetDataHandle().GetData())
+					{
+						if (Data->GetObject() == InObject)
+						{
+							OutNodePtr = CurNodePtr;
+						}
+						// Handle BP inherited subobjects that are on instances
+						else if (!OutNodePtr && Data->GetObjectForBlueprint(BP) == InObject->GetArchetype())
+						{
+							OutNodePtr = CurNodePtr;
+						}
+					}
                 }
             });
 
