Correctly check if a dragged node is a child actor subtree object instead of just a child actor node. This fixes the tooltip when dragging into the blueprint event graph

#rnx

diff --git a/Engine/Source/Editor/SubobjectEditor/Private/SSubobjectEditor.cpp b/Engine/Source/Editor/SubobjectEditor/Private/SSubobjectEditor.cpp
--- a/Engine/Source/Editor/SubobjectEditor/Private/SSubobjectEditor.cpp
+++ b/Engine/Source/Editor/SubobjectEditor/Private/SSubobjectEditor.cpp
@@ -75,7 +75,7 @@ namespace SubobjectEditorUtils
 		for (FSubobjectEditorTreeNodePtrType NodePtr : InNodePtrs)
 		{
 			const FSubobjectData* Data = NodePtr ? NodePtr->GetDataSource() : nullptr;
-			if (Data && Data->IsChildActor())
+			if (Data && Data->IsChildActorSubtreeObject())
 			{
 				return true;
 			}
