Make sure that a newly pasted node is scrolled into view

#rnx

diff --git a/Engine/Source/Editor/SubobjectEditor/Private/SSubobjectBlueprintEditor.cpp b/Engine/Source/Editor/SubobjectEditor/Private/SSubobjectBlueprintEditor.cpp
--- a/Engine/Source/Editor/SubobjectEditor/Private/SSubobjectBlueprintEditor.cpp
+++ b/Engine/Source/Editor/SubobjectEditor/Private/SSubobjectBlueprintEditor.cpp
@@ -361,6 +361,7 @@ void SSubobjectBlueprintEditor::PasteNodes()
 			{
 				if(FSubobjectEditorTreeNodePtrType SlateNode = FindSlateNodeForHandle(Handle))
 				{
+					TreeWidget->RequestScrollIntoView(SlateNode);
 					TreeWidget->SetItemSelection(SlateNode, true);	
 				}
 			}
diff --git a/Engine/Source/Editor/SubobjectEditor/Private/SSubobjectInstanceEditor.cpp b/Engine/Source/Editor/SubobjectEditor/Private/SSubobjectInstanceEditor.cpp
--- a/Engine/Source/Editor/SubobjectEditor/Private/SSubobjectInstanceEditor.cpp
+++ b/Engine/Source/Editor/SubobjectEditor/Private/SSubobjectInstanceEditor.cpp
@@ -278,6 +278,7 @@ void SSubobjectInstanceEditor::PasteNodes()
 			{
 				if(FSubobjectEditorTreeNodePtrType SlateNode = FindSlateNodeForHandle(Handle))
 				{
+					TreeWidget->RequestScrollIntoView(SlateNode);
 					TreeWidget->SetItemSelection(SlateNode, true);	
 				}
 			}
