Check for an instance inherited default scene root when disallowing

#rnx

diff --git a/Engine/Source/Editor/SubobjectEditor/Private/SSubobjectEditor.cpp b/Engine/Source/Editor/SubobjectEditor/Private/SSubobjectEditor.cpp
--- a/Engine/Source/Editor/SubobjectEditor/Private/SSubobjectEditor.cpp
+++ b/Engine/Source/Editor/SubobjectEditor/Private/SSubobjectEditor.cpp
@@ -1050,7 +1050,7 @@ void SSubobject_RowWidget::HandleOnDragEnter(const FDragDropEvent& DragDropEvent
 						&& DraggedTemplate->Mobility >= HoveredTemplate->Mobility
 						&& (!HoveredTemplate->IsEditorOnly() || DraggedTemplate->IsEditorOnly());
 
-					if (!NodePtr->CanReparent() && (!NodeData->IsDefaultSceneRoot() || NodeData->IsInheritedComponent()))
+					if (!NodePtr->CanReparent() && (!NodeData->IsDefaultSceneRoot() || NodeData->IsInheritedComponent() || NodeData->IsInstancedInheritedComponent()))
 					{
 						// Cannot make the dropped node the new root if we cannot reparent the current root
 						Message = LOCTEXT("DropActionToolTip_Error_CannotReparentRootNode", "The root component in this Blueprint is inherited and cannot be replaced.");
