Add OnHighlightPropertyInDetailsView to the blueprint editor to account for the behavior required in the SCS diff tree

#rnx

diff --git a/Engine/Source/Editor/SubobjectEditor/Private/SSubobjectBlueprintEditor.cpp b/Engine/Source/Editor/SubobjectEditor/Private/SSubobjectBlueprintEditor.cpp
--- a/Engine/Source/Editor/SubobjectEditor/Private/SSubobjectBlueprintEditor.cpp
+++ b/Engine/Source/Editor/SubobjectEditor/Private/SSubobjectBlueprintEditor.cpp
@@ -9,6 +9,7 @@
 #include "ISCSEditorUICustomization.h"	// #TODO_BH Rename this to subobject
 #include "SCSEditorExtensionContext.h"	// #TODO_BH Rename this to subobject (this one will break VK, for sure)
 #include "ToolMenus.h"
+#include "PropertyPath.h"
 
 #include "Styling/SlateIconFinder.h"
 #include "SlateOptMacros.h"
@@ -42,7 +43,9 @@ void SSubobjectBlueprintEditor::Construct(const FArguments& InArgs)
 	PreviewActor = InArgs._PreviewActor;
 	OnSelectionUpdated = InArgs._OnSelectionUpdated;
 	OnItemDoubleClicked = InArgs._OnItemDoubleClicked;
+	OnHighlightPropertyInDetailsView = InArgs._OnHighlightPropertyInDetailsView;
 	AllowEditing = InArgs._AllowEditing;
+	HideComponentClassCombo = InArgs._HideComponentClassCombo;
 	bAllowTreeUpdates = true;
 	
 	CreateCommandList();
@@ -752,6 +755,50 @@ void SSubobjectBlueprintEditor::BuildMenuEventsSection(FMenuBuilder& Menu, UBlue
 	Menu.EndSection();
 }
 
+void SSubobjectBlueprintEditor::HighlightTreeNode(FName TreeNodeName, const FPropertyPath& Property)
+{
+	// If there are no nodes then there is nothing to do!
+	if (RootNodes.Num() == 0)
+	{
+		return;
+	}
+
+	// Find the tree node that matches up with the given property FName
+	FSubobjectEditorTreeNodePtrType FoundNode = nullptr;
+
+	TSet<FSubobjectEditorTreeNodePtrType> VisitedNodes;
+	DepthFirstTraversal(RootNodes[0], VisitedNodes,
+		[&FoundNode, TreeNodeName](
+		const FSubobjectEditorTreeNodePtrType& CurNodePtr)
+		{
+			if (CurNodePtr->GetDataHandle().IsValid())
+			{
+				if (CurNodePtr->GetDataHandle().GetSharedDataPtr()->GetVariableName() == TreeNodeName)
+				{
+					FoundNode = CurNodePtr;
+				}
+			}
+		}
+	);
+
+
+	// If it was found, then select it and broadcast the delegate to let the diff panel know
+	if (FoundNode)
+	{
+		SelectNode(FoundNode->AsShared(), false);
+
+		if (Property != FPropertyPath())
+		{
+			// Invoke the delegate to highlight the property
+			OnHighlightPropertyInDetailsView.ExecuteIfBound(Property);
+		}
+
+		return;
+	}
+
+	ClearSelection();
+}
+
 void SSubobjectBlueprintEditor::CreateEventsForSelection(UBlueprint* Blueprint, FName EventName, FGetSelectedObjectsDelegate GetSelectedObjectsDelegate)
 {
 	if (EventName != NAME_None)
diff --git a/Engine/Source/Editor/SubobjectEditor/Public/SSubobjectBlueprintEditor.h b/Engine/Source/Editor/SubobjectEditor/Public/SSubobjectBlueprintEditor.h
--- a/Engine/Source/Editor/SubobjectEditor/Public/SSubobjectBlueprintEditor.h
+++ b/Engine/Source/Editor/SubobjectEditor/Public/SSubobjectBlueprintEditor.h
@@ -9,19 +9,26 @@
 */
 class SUBOBJECTEDITOR_API SSubobjectBlueprintEditor final : public SSubobjectEditor
 {
+public:
+	DECLARE_DELEGATE_OneParam(FOnHighlightPropertyInDetailsView, const class FPropertyPath&);
+
 private:	
 	SLATE_BEGIN_ARGS(SSubobjectBlueprintEditor)
         : _ObjectContext(nullptr)
 		, _PreviewActor(nullptr)
         , _AllowEditing(true)
+		, _HideComponentClassCombo(false)
         , _OnSelectionUpdated()
+		, _OnHighlightPropertyInDetailsView()
 		{}
 
 		SLATE_ATTRIBUTE(UObject*, ObjectContext)
 		SLATE_ATTRIBUTE(AActor*, PreviewActor)
 	    SLATE_ATTRIBUTE(bool, AllowEditing)
+		SLATE_ATTRIBUTE(bool, HideComponentClassCombo)
 	    SLATE_EVENT(FOnSelectionUpdated, OnSelectionUpdated)
 	    SLATE_EVENT(FOnItemDoubleClicked, OnItemDoubleClicked)
+		SLATE_EVENT(FOnHighlightPropertyInDetailsView, OnHighlightPropertyInDetailsView)
 	    
 	SLATE_END_ARGS()
 	
@@ -58,6 +65,9 @@ public:
 
 public:
 
+	/** Delegate to invoke when the given property should be highlighted in the details view (e.g. diff). */
+	FOnHighlightPropertyInDetailsView OnHighlightPropertyInDetailsView;
+	
 	/**
 	* Fills out an events section in ui.
 	* @param Menu								the menu to add the events section into
@@ -68,6 +78,14 @@ public:
 	*/
 	static void BuildMenuEventsSection(FMenuBuilder& Menu, UBlueprint* Blueprint, UClass* SelectedClass, FCanExecuteAction CanExecuteActionDelegate, FGetSelectedObjectsDelegate GetSelectedObjectsDelegate);
 
+	/**
+	* Highlight a tree node and, optionally, a property with in it
+	*
+	* @param TreeNodeName		Name of the treenode to be highlighted
+	* @param Property	The name of the property to be highlighted in the details view
+	*/
+	void HighlightTreeNode(FName TreeNodeName, const class FPropertyPath& Property);
+	
 protected:
 
 	/**
