Subobject editor recognizes HideAddComponentButton of ISCSEditorUICustomization.

diff --git a/Engine/Source/Editor/SubobjectEditor/Private/SSubobjectBlueprintEditor.cpp b/Engine/Source/Editor/SubobjectEditor/Private/SSubobjectBlueprintEditor.cpp
--- a/Engine/Source/Editor/SubobjectEditor/Private/SSubobjectBlueprintEditor.cpp
+++ b/Engine/Source/Editor/SubobjectEditor/Private/SSubobjectBlueprintEditor.cpp
@@ -80,7 +80,7 @@ void SSubobjectBlueprintEditor::Construct(const FArguments& InArgs)
 	[
 		SNew(SComponentClassCombo)
 		.AddMetaData<FTagMetaData>(FTagMetaData(TEXT("Actor.AddComponent")))
-		.Visibility(HideComponentClassCombo.Get() ? EVisibility::Hidden : EVisibility::Visible)
+		.Visibility(this, &SSubobjectBlueprintEditor::GetComponentClassComboButtonVisibility)
 		.OnSubobjectClassSelected(this, &SSubobjectBlueprintEditor::PerformComboAddClass)
 		.ToolTipText(LOCTEXT("AddComponent_Tooltip", "Adds a new component to this actor"))
 		.IsEnabled(true)
diff --git a/Engine/Source/Editor/SubobjectEditor/Private/SSubobjectEditor.cpp b/Engine/Source/Editor/SubobjectEditor/Private/SSubobjectEditor.cpp
--- a/Engine/Source/Editor/SubobjectEditor/Private/SSubobjectEditor.cpp
+++ b/Engine/Source/Editor/SubobjectEditor/Private/SSubobjectEditor.cpp
@@ -3382,6 +3382,13 @@ EVisibility SSubobjectEditor::GetEditBlueprintButtonVisibility() const
 		       : EVisibility::Visible;
 }
 
+EVisibility SSubobjectEditor::GetComponentClassComboButtonVisibility() const
+{
+	return (HideComponentClassCombo.Get() 
+	|| (UICustomization.IsValid() && UICustomization->HideAddComponentButton())) 
+	? EVisibility::Collapsed : EVisibility::Visible;
+}
+
 void SSubobjectEditor::Utils::PopulateHandlesArray(const TArray<FSubobjectEditorTreeNodePtrType>& SlateNodePtrs, TArray<FSubobjectDataHandle>& OutHandles)
 {
 	for(const FSubobjectEditorTreeNodePtrType& DroppedNodePtr : SlateNodePtrs)
diff --git a/Engine/Source/Editor/SubobjectEditor/Public/SSubobjectEditor.h b/Engine/Source/Editor/SubobjectEditor/Public/SSubobjectEditor.h
--- a/Engine/Source/Editor/SubobjectEditor/Public/SSubobjectEditor.h
+++ b/Engine/Source/Editor/SubobjectEditor/Public/SSubobjectEditor.h
@@ -714,6 +714,9 @@ protected:
 	/** @return The visibility of the Edit Blueprint button (only visible with an actor instance that is created from a blueprint)*/
 	virtual EVisibility GetEditBlueprintButtonVisibility() const;
 
+	/** @return The visibility of the Add Component combo button */
+	virtual EVisibility GetComponentClassComboButtonVisibility() const;
+
 	/** If true, then the tree widget will clear selection on click. */
 	virtual bool ClearSelectionOnClick() const { return false; }
 	
