Correctly use the parent handle when adding a component with a selection from the content browser in the subobject editor. Previously it was attempting to find the parent handle for the asset, which would always return null. We just need to use the root handle of the subobject tree instead.

#rnx

diff --git a/Engine/Source/Editor/SubobjectEditor/Private/SSubobjectEditor.cpp b/Engine/Source/Editor/SubobjectEditor/Private/SSubobjectEditor.cpp
--- a/Engine/Source/Editor/SubobjectEditor/Private/SSubobjectEditor.cpp
+++ b/Engine/Source/Editor/SubobjectEditor/Private/SSubobjectEditor.cpp
@@ -3291,13 +3291,13 @@ FSubobjectDataHandle SSubobjectEditor::PerformComboAddClass(TSubclassOf<UActorCo
 				{
 					if (ComponentClasses[ComponentIndex]->IsChildOf(NewClass))
 					{
-						const FSubobjectDataHandle ParentHandle = SubobjectSystem->FindHandleForObject(CachedRootHandle, Object);
-						if(ParentHandle.IsValid())
-						{
-							NewComponentHandle = AddNewSubobject(ParentHandle, NewClass, Object, OutFailReason, MoveTemp(AddTransaction));
-							bAddedComponent = true;
-							break;
-						}
+						const TArray<FSubobjectDataHandle>& SelectedHandles = GetSelectedHandles();
+						const FSubobjectDataHandle& ParentHandle = SelectedHandles.Num() > 0 ? SelectedHandles[0] : CachedRootHandle;
+						ensureMsgf(ParentHandle.IsValid(), TEXT("Attempting to add a component from an invalid selection!"));
+						
+						NewComponentHandle = AddNewSubobject(ParentHandle, NewClass, Object, OutFailReason, MoveTemp(AddTransaction));
+						bAddedComponent = true;
+						break;
 					}
 				}
 			}
