Update Component Tree icons and styling

diff --git a/Engine/Source/Editor/EditorStyle/Private/StarshipStyle.cpp b/Engine/Source/Editor/EditorStyle/Private/StarshipStyle.cpp
--- a/Engine/Source/Editor/EditorStyle/Private/StarshipStyle.cpp
+++ b/Engine/Source/Editor/EditorStyle/Private/StarshipStyle.cpp
@@ -5289,8 +5289,8 @@ void FStarshipEditorStyle::FStyle::SetupPersonaStyle()
 		Set( "BlueprintEditor.ShowGrid", new IMAGE_BRUSH_SVG( "Starship/Common/Grid", Icon16x16 ) );
 		Set( "BlueprintEditor.EnableSimulation", new IMAGE_BRUSH_SVG("Starship/MainToolbar/simulate", Icon20x20));
 		Set( "BlueprintEditor.EnableProfiling", new IMAGE_BRUSH_SVG("Starship/Common/Statistics", Icon20x20) );
-		Set( "SCS.NativeComponent", new IMAGE_BRUSH( "Icons/NativeSCSComponent", Icon20x20 ));
-		Set( "SCS.Component", new IMAGE_BRUSH( "Icons/SCSComponent", Icon20x20 ));
+		Set( "SCS.NativeComponent", new IMAGE_BRUSH_SVG("Starship/AssetIcons/SceneComponent", Icon16x16));
+		Set( "SCS.Component", new IMAGE_BRUSH_SVG("Starship/AssetIcons/ActorComponent", Icon16x16));
 
 		// curve viewer
 		Set("AnimCurveViewer.MorphTargetOn", new IMAGE_BRUSH(TEXT("Persona/AnimCurveViewer/MorphTarget_On"), Icon16x16));
diff --git a/Engine/Source/Editor/SubobjectEditor/Private/SSubobjectEditor.cpp b/Engine/Source/Editor/SubobjectEditor/Private/SSubobjectEditor.cpp
--- a/Engine/Source/Editor/SubobjectEditor/Private/SSubobjectEditor.cpp
+++ b/Engine/Source/Editor/SubobjectEditor/Private/SSubobjectEditor.cpp
@@ -33,6 +33,7 @@
 #include "Widgets/SToolTip.h"
 #include "Widgets/Input/SHyperlink.h"
 #include "Widgets/Layout/SSpacer.h"
+#include "Widgets/Layout/SSeparator.h"
 #include "Selection.h"
 #include "Framework/MultiBox/MultiBoxBuilder.h"
 #include "SComponentClassCombo.h"
@@ -734,16 +735,13 @@ TSharedRef<SWidget> SSubobject_RowWidget::GenerateWidgetForColumn(const FName& C
 	
 	if (Node->IsSeperator())
 	{
-		static const FLinearColor Debug(1.0f, 0.0f, 0.0f);
-		return SNew(SSpacer);
-		// return SNew(SBox)
-		// 	.Padding(1.f)
-		// 	[
-		// 		SNew(SBorder)
-		// 		.Padding(FEditorStyle::GetMargin(TEXT("Menu.Separator.Padding")))
-		// 		.BorderImage(FEditorStyle::GetBrush(TEXT("Menu.Separator")))
-		// 		.ColorAndOpacity(Debug)
-		// 	];
+		return SNew(SBox)
+			.Padding(1.f)
+			[
+				SNew(SSeparator)
+				.SeparatorImage(FAppStyle::Get().GetBrush("Menu.Separator"))
+				.Thickness(1.f)
+		 	];
 	}
 
 	FSubobjectData* Data = Node->GetDataSource();
@@ -2473,10 +2471,6 @@ void SSubobjectEditor::UpdateTree(bool bRegenerateTreeNodes /* = true */)
 			RootNodes.Add(Node);
 			AddedNodes.Add(Node->GetDataHandle(), Node);
 			CachedRootHandle = Node->GetDataHandle();
-
-			SeperatorNode = MakeShareable<FSubobjectEditorTreeNode>(
-				new FSubobjectEditorTreeNode(FSubobjectDataHandle::InvalidHandle, /** bIsSeperator */true));
-			RootNodes.Add(SeperatorNode);
 			
 			TreeWidget->SetItemExpansion(Node, true);
 			TreeWidget->SetItemExpansion(SeperatorNode, true);
@@ -2503,6 +2497,16 @@ void SSubobjectEditor::UpdateTree(bool bRegenerateTreeNodes /* = true */)
 				check(ParentNode);
 				ParentNode->AddChild(NewNode);
 				TreeWidget->SetItemExpansion(ParentNode, true);			
+
+				// Add a seperator after the default scene root, but only if there are more items below it
+				if (Data->IsDefaultSceneRoot() && Data->IsInheritedComponent() && SubobjectData.Find(Handle) < SubobjectData.Num() - 1)
+				{
+					SeperatorNode = MakeShareable<FSubobjectEditorTreeNode>(
+						new FSubobjectEditorTreeNode(FSubobjectDataHandle::InvalidHandle, /** bIsSeperator */true));
+					AddedNodes.Add(SeperatorNode->GetDataHandle(), SeperatorNode);
+					ParentNode->AddChild(SeperatorNode);
+				}
+				
 			}
 			TreeWidget->SetItemExpansion(NewNode, true);
 		}
