Fix an ensure that is thrown when searching a component tree that has a brand new component added to it. This is an outdated ensure after the refactor from SCSEditor -> Subobject Editor. Remove an unnecessary include file

#rnx

diff --git a/Engine/Source/Editor/SubobjectEditor/Private/SSubobjectEditor.cpp b/Engine/Source/Editor/SubobjectEditor/Private/SSubobjectEditor.cpp
--- a/Engine/Source/Editor/SubobjectEditor/Private/SSubobjectEditor.cpp
+++ b/Engine/Source/Editor/SubobjectEditor/Private/SSubobjectEditor.cpp
@@ -16,7 +16,6 @@
 #include "Editor/EditorEngine.h"
 #include "ISCSEditorUICustomization.h"	// #TODO_BH Rename this to subobject
 #include "SubobjectEditorMenuContext.h"
-#include "SSCSEditorMenuContext.h"		// #TODO_BH Rename this to subobject
 #include "Toolkits/ToolkitManager.h"
 #include "BlueprintEditor.h"
 #include "ToolMenus.h"
@@ -301,8 +300,7 @@ FSubobjectEditorTreeNodePtrType FSubobjectEditorTreeNode::FindChild(const FSubob
 
 bool FSubobjectEditorTreeNode::IsFlaggedForFiltration() const
 {
-	return ensureMsgf(FilterFlags != EFilteredState::Unknown, TEXT("Querying a bad filtration state.")) ?
-		(FilterFlags & EFilteredState::FilteredInMask) == 0 : false;
+	return FilterFlags != EFilteredState::Unknown ? (FilterFlags & EFilteredState::FilteredInMask) == 0 : false;
 }
 
 void FSubobjectEditorTreeNode::SetCachedFilterState(bool bMatchesFilter, bool bUpdateParent)
