Fixing empty Actors having a root component with an Edit in Blueprint hyperlink that doesn't work

diff --git a/Engine/Source/Editor/SubobjectEditor/Private/SSubobjectEditor.cpp b/Engine/Source/Editor/SubobjectEditor/Private/SSubobjectEditor.cpp
--- a/Engine/Source/Editor/SubobjectEditor/Private/SSubobjectEditor.cpp
+++ b/Engine/Source/Editor/SubobjectEditor/Private/SSubobjectEditor.cpp
@@ -1741,14 +1741,15 @@ TSharedRef<SWidget> SSubobject_RowWidget::GetInheritedLinkWidget()
 	// If the subobject is inherited and not native then it must be from a blueprint
 	else if(Data->IsInstancedInheritedComponent() || Data->IsBlueprintInheritedComponent())
 	{
-		const UBlueprint* BP = Data->GetBlueprint();
-		static const FText InheritedBPLabel = LOCTEXT("InheritedBpLabel", "Edit in Blueprint");
-		
-		return SNew(SHyperlink)
-			.Style(FAppStyle::Get(), "Common.GotoBlueprintHyperlink")
-			.OnNavigate(this, &SSubobject_RowWidget::OnEditBlueprintClicked)
-			.Text(InheritedBPLabel)
-			.ToolTipText(LOCTEXT("EditBlueprint_ToolTip", "Click to edit the blueprint"));
+		if(const UBlueprint* BP = Data->GetBlueprint())
+		{
+			static const FText InheritedBPLabel = LOCTEXT("InheritedBpLabel", "Edit in Blueprint");
+			return SNew(SHyperlink)
+				.Style(FAppStyle::Get(), "Common.GotoBlueprintHyperlink")
+				.OnNavigate(this, &SSubobject_RowWidget::OnEditBlueprintClicked)
+				.Text(InheritedBPLabel)
+				.ToolTipText(LOCTEXT("EditBlueprint_ToolTip", "Click to edit the blueprint"));
+		}
 	}
 
 	// Non-inherited subobjects shouldn't show anything! 
