Hide C++ related menus and hyperlinks from the subobject editor when C++ is not allowed

diff --git a/Engine/Source/Editor/SubobjectEditor/Private/SSubobjectEditor.cpp b/Engine/Source/Editor/SubobjectEditor/Private/SSubobjectEditor.cpp
--- a/Engine/Source/Editor/SubobjectEditor/Private/SSubobjectEditor.cpp
+++ b/Engine/Source/Editor/SubobjectEditor/Private/SSubobjectEditor.cpp
@@ -14,6 +14,9 @@
 #include "Kismet2/ChildActorComponentEditorUtils.h"
 #include "GraphEditorActions.h"
 #include "Editor/EditorEngine.h"
+#include "Editor/UnrealEdEngine.h"
+#include "Preferences/UnrealEdOptions.h"
+#include "UnrealEdGlobals.h"
 #include "ISCSEditorUICustomization.h"	// #TODO_BH Rename this to subobject
 #include "SubobjectEditorMenuContext.h"
 #include "Toolkits/ToolkitManager.h"
@@ -1736,7 +1739,8 @@ TSharedRef<SWidget> SSubobject_RowWidget::GetInheritedLinkWidget()
 			.Style(FAppStyle::Get(), "Common.GotoNativeCodeHyperlink")
 			.OnNavigate(this, &SSubobject_RowWidget::OnEditNativeCppClicked)
 			.Text(NativeCppLabel)
-			.ToolTipText(FText::Format(LOCTEXT("GoToCode_ToolTip", "Click to open this source file in {0}"), FSourceCodeNavigation::GetSelectedSourceCodeIDE()));	
+			.ToolTipText(FText::Format(LOCTEXT("GoToCode_ToolTip", "Click to open this source file in {0}"), FSourceCodeNavigation::GetSelectedSourceCodeIDE()))
+			.Visibility(this, &SSubobject_RowWidget::GetEditNativeCppVisibility);
 	}
 	// If the subobject is inherited and not native then it must be from a blueprint
 	else if(Data->IsInstancedInheritedComponent() || Data->IsBlueprintInheritedComponent())
@@ -1792,6 +1796,11 @@ void SSubobject_RowWidget::OnEditBlueprintClicked()
 	}
 }
 
+EVisibility SSubobject_RowWidget::GetEditNativeCppVisibility() const
+{
+	return ensure(GUnrealEd) && GUnrealEd->GetUnrealEdOptions()->IsCPPAllowed() ? EVisibility::Visible : EVisibility::Collapsed;
+}
+
 void SSubobject_RowWidget::OnEditNativeCppClicked()
 {
 	FSubobjectEditorTreeNodePtrType NodePtr = GetSubobjectPtr();
diff --git a/Engine/Source/Editor/SubobjectEditor/Public/SSubobjectEditor.h b/Engine/Source/Editor/SubobjectEditor/Public/SSubobjectEditor.h
--- a/Engine/Source/Editor/SubobjectEditor/Public/SSubobjectEditor.h
+++ b/Engine/Source/Editor/SubobjectEditor/Public/SSubobjectEditor.h
@@ -321,6 +321,8 @@ private:
 
 	/** Callback used when the user clicks on a blueprint inherited variable */
 	void OnEditBlueprintClicked();
+
+	EVisibility GetEditNativeCppVisibility() const;
 	
 	/** Callback used when the user clicks on a native inherited variable */
 	void OnEditNativeCppClicked();
diff --git a/Engine/Source/Editor/UnrealEd/Private/Kismet2/ComponentEditorUtils.cpp b/Engine/Source/Editor/UnrealEd/Private/Kismet2/ComponentEditorUtils.cpp
--- a/Engine/Source/Editor/UnrealEd/Private/Kismet2/ComponentEditorUtils.cpp
+++ b/Engine/Source/Editor/UnrealEd/Private/Kismet2/ComponentEditorUtils.cpp
@@ -11,9 +11,7 @@
 #include "Components/MeshComponent.h"
 #include "Exporters/Exporter.h"
 #include "Materials/Material.h"
-#include "Editor/UnrealEdEngine.h"
 #include "Components/DecalComponent.h"
-#include "UnrealEdGlobals.h"
 #include "ScopedTransaction.h"
 #include "EdGraphSchema_K2.h"
 #include "Kismet2/BlueprintEditorUtils.h"
@@ -25,6 +23,10 @@
 #include "Styling/SlateIconFinder.h"
 #include "Framework/MultiBox/MultiBoxBuilder.h"
 #include "Editor.h"
+#include "Editor/EditorEngine.h"
+#include "Editor/UnrealEdEngine.h"
+#include "Preferences/UnrealEdOptions.h"
+#include "UnrealEdGlobals.h"
 
 #include "HAL/PlatformApplicationMisc.h"
 #include "ToolMenus.h"
@@ -1105,6 +1107,7 @@ void FComponentEditorUtils::FillComponentContextMenuOptions(UToolMenu* Menu, con
 		}
 		else
 		{
+			if (ensure(GUnrealEd) && GUnrealEd->GetUnrealEdOptions()->IsCPPAllowed())
 			{
 				FToolMenuSection& Section = Menu->AddSection("ComponentCode", LOCTEXT("ComponentCodeHeading", "C++"));
 				if (FSourceCodeNavigation::IsCompilerAvailable())
