SActorDetails: Fix SubObjectEditor not being cleared properly

#rnx

diff --git a/Engine/Source/Editor/LevelEditor/Private/SActorDetails.cpp b/Engine/Source/Editor/LevelEditor/Private/SActorDetails.cpp
--- a/Engine/Source/Editor/LevelEditor/Private/SActorDetails.cpp
+++ b/Engine/Source/Editor/LevelEditor/Private/SActorDetails.cpp
@@ -488,8 +488,7 @@ void SActorDetails::RefreshTopLevelElements(TArrayView<const TTypedElement<IType
 	// Update the underlying details view and the Elements buttons
 	SetElementDetailsObjects(TopLevelElements, bForceRefresh, bOverrideLock, &InDetailsElements);
 
-	// Update the Subobject tree if we were asked to edit a single actor
-	if (AActor* Actor = GetActorContext())
+	// Update the Subobject tree
 	{
 		// Enable the selection guard to prevent OnTreeSelectionChanged() from altering the editor's component selection
 		TGuardValue<bool> SelectionGuard(bSelectionGuard, true);
diff --git a/Engine/Source/Editor/SubobjectEditor/Private/SSubobjectEditor.cpp b/Engine/Source/Editor/SubobjectEditor/Private/SSubobjectEditor.cpp
--- a/Engine/Source/Editor/SubobjectEditor/Private/SSubobjectEditor.cpp
+++ b/Engine/Source/Editor/SubobjectEditor/Private/SSubobjectEditor.cpp
@@ -2459,14 +2459,8 @@ void SSubobjectEditor::UpdateTree(bool bRegenerateTreeNodes /* = true */)
 		return;
 	}
 
-	UObject* Context = GetObjectContext();
-
-	if (bRegenerateTreeNodes && Context)
+	if (bRegenerateTreeNodes)
 	{
-		// Obtain the set of expandable tree nodes that are currently collapsed
-		TSet<FSubobjectEditorTreeNodePtrType> CollapsedTreeNodes;
-		GetCollapsedNodes(RootNodes.Num() > 0 ? RootNodes[0] : FSubobjectEditorTreeNodePtrType(), CollapsedTreeNodes);
-
 		// Obtain the list of selected items
 		TArray<FSubobjectEditorTreeNodePtrType> SelectedTreeNodes = GetSelectedNodes();
 
@@ -2477,75 +2471,78 @@ void SSubobjectEditor::UpdateTree(bool bRegenerateTreeNodes /* = true */)
 		}
 		RootNodes.Empty();
 
-		ensureMsgf(FModuleManager::Get().LoadModule("SubobjectDataInterface"), TEXT("The Subobject Data Interface module is required."));
+		if (UObject* Context = GetObjectContext())
+		{
+			ensureMsgf(FModuleManager::Get().LoadModule("SubobjectDataInterface"), TEXT("The Subobject Data Interface module is required."));
 
-		USubobjectDataSubsystem* DataSubsystem = USubobjectDataSubsystem::Get();
-		check(DataSubsystem);
+			USubobjectDataSubsystem* DataSubsystem = USubobjectDataSubsystem::Get();
+			check(DataSubsystem);
 
-		TArray<FSubobjectDataHandle> SubobjectData;
-				
-		DataSubsystem->GatherSubobjectData(Context, SubobjectData);
-		
-		FSubobjectEditorTreeNodePtrType SeperatorNode;
-		TMap<FSubobjectDataHandle, FSubobjectEditorTreeNodePtrType> AddedNodes;
+			TArray<FSubobjectDataHandle> SubobjectData;
 
-		// By default, root node will always be expanded. If possible we will restore the collapsed state later.
-		if (SubobjectData.Num() > 0)
-		{
-			FSubobjectEditorTreeNodePtrType Node = MakeShareable<FSubobjectEditorTreeNode>(
-				new FSubobjectEditorTreeNode(SubobjectData[0]));
-			RootNodes.Add(Node);
-			AddedNodes.Add(Node->GetDataHandle(), Node);
-			CachedRootHandle = Node->GetDataHandle();
-			
-			TreeWidget->SetItemExpansion(Node, true);
-			TreeWidget->SetItemExpansion(SeperatorNode, true);
+			DataSubsystem->GatherSubobjectData(Context, SubobjectData);
 
-			RefreshFilteredState(Node, false);
-		}
+			FSubobjectEditorTreeNodePtrType SeperatorNode;
+			TMap<FSubobjectDataHandle, FSubobjectEditorTreeNodePtrType> AddedNodes;
 
-		// Create slate nodes for each subobject
-		for (FSubobjectDataHandle& Handle : SubobjectData)
-		{
-			// Do we have a slate node for this handle already? If not, then we need to make one
-			FSubobjectEditorTreeNodePtrType NewNode = SSubobjectEditor::FindOrCreateSlateNodeForHandle(Handle, AddedNodes);
-			
-			FSubobjectData* Data = Handle.GetData();
+			// By default, root node will always be expanded. If possible we will restore the collapsed state later.
+			if (SubobjectData.Num() > 0)
+			{
+				FSubobjectEditorTreeNodePtrType Node = MakeShareable<FSubobjectEditorTreeNode>(
+					new FSubobjectEditorTreeNode(SubobjectData[0]));
+				RootNodes.Add(Node);
+				AddedNodes.Add(Node->GetDataHandle(), Node);
+				CachedRootHandle = Node->GetDataHandle();
+
+				TreeWidget->SetItemExpansion(Node, true);
+				TreeWidget->SetItemExpansion(SeperatorNode, true);
 
-			const FSubobjectDataHandle& ParentHandle = Data->GetParentHandle();
+				RefreshFilteredState(Node, false);
+			}
 
-			// Have parent? 
-			if (ParentHandle.IsValid())
+			// Create slate nodes for each subobject
+			for (FSubobjectDataHandle& Handle : SubobjectData)
 			{
-				// Get the parent node for this subobject
-				FSubobjectEditorTreeNodePtrType ParentNode = SSubobjectEditor::FindOrCreateSlateNodeForHandle(ParentHandle, AddedNodes);
+				// Do we have a slate node for this handle already? If not, then we need to make one
+				FSubobjectEditorTreeNodePtrType NewNode = SSubobjectEditor::FindOrCreateSlateNodeForHandle(Handle, AddedNodes);
+
+				FSubobjectData* Data = Handle.GetData();
 
-				check(ParentNode);
-				ParentNode->AddChild(NewNode);
-				TreeWidget->SetItemExpansion(ParentNode, true);			
+				const FSubobjectDataHandle& ParentHandle = Data->GetParentHandle();
 
-				// Add a seperator after the default scene root, but only if there are more items below it
-				if (Data->IsDefaultSceneRoot() && Data->IsInheritedComponent() && SubobjectData.Find(Handle) < SubobjectData.Num() - 1)
+				// Have parent? 
+				if (ParentHandle.IsValid())
 				{
-					SeperatorNode = MakeShareable<FSubobjectEditorTreeNode>(
-						new FSubobjectEditorTreeNode(FSubobjectDataHandle::InvalidHandle, /** bIsSeperator */true));
-					AddedNodes.Add(SeperatorNode->GetDataHandle(), SeperatorNode);
-					ParentNode->AddChild(SeperatorNode);
+					// Get the parent node for this subobject
+					FSubobjectEditorTreeNodePtrType ParentNode = SSubobjectEditor::FindOrCreateSlateNodeForHandle(ParentHandle, AddedNodes);
+
+					check(ParentNode);
+					ParentNode->AddChild(NewNode);
+					TreeWidget->SetItemExpansion(ParentNode, true);
+
+					// Add a seperator after the default scene root, but only if there are more items below it
+					if (Data->IsDefaultSceneRoot() && Data->IsInheritedComponent() && SubobjectData.Find(Handle) < SubobjectData.Num() - 1)
+					{
+						SeperatorNode = MakeShareable<FSubobjectEditorTreeNode>(
+							new FSubobjectEditorTreeNode(FSubobjectDataHandle::InvalidHandle, /** bIsSeperator */true));
+						AddedNodes.Add(SeperatorNode->GetDataHandle(), SeperatorNode);
+						ParentNode->AddChild(SeperatorNode);
+					}
+
 				}
-				
+				TreeWidget->SetItemExpansion(NewNode, true);
 			}
-			TreeWidget->SetItemExpansion(NewNode, true);
-		}
 
-		RestoreSelectionState(SelectedTreeNodes);
+			RestoreSelectionState(SelectedTreeNodes);
 
-		// If we have a pending deferred rename request, redirect it to the new tree node
-		if(DeferredRenameRequest.IsValid())
-		{
-			FSubobjectEditorTreeNodePtrType NodeToRenamePtr = FindSlateNodeForHandle(DeferredRenameRequest);
-			if(NodeToRenamePtr.IsValid())
+			// If we have a pending deferred rename request, redirect it to the new tree node
+			if (DeferredRenameRequest.IsValid())
 			{
-				TreeWidget->RequestScrollIntoView(NodeToRenamePtr);
+				FSubobjectEditorTreeNodePtrType NodeToRenamePtr = FindSlateNodeForHandle(DeferredRenameRequest);
+				if (NodeToRenamePtr.IsValid())
+				{
+					TreeWidget->RequestScrollIntoView(NodeToRenamePtr);
+				}
 			}
 		}
 	}
