PR #9595: Fix selected component being lost when compiling a Blueprint (Contributed by [REDACTED_PERSON])

diff --git a/Engine/Source/Editor/SubobjectEditor/Private/SSubobjectEditor.cpp b/Engine/Source/Editor/SubobjectEditor/Private/SSubobjectEditor.cpp
--- a/Engine/Source/Editor/SubobjectEditor/Private/SSubobjectEditor.cpp
+++ b/Engine/Source/Editor/SubobjectEditor/Private/SSubobjectEditor.cpp
@@ -2581,7 +2581,7 @@ void SSubobjectEditor::ClearSelection()
 	}
 }
 
-void SSubobjectEditor::RestoreSelectionState(TArray<FSubobjectEditorTreeNodePtrType>& SelectedTreeNodes)
+void SSubobjectEditor::RestoreSelectionState(TArray<FSubobjectEditorTreeNodePtrType>& SelectedTreeNodes, bool bFallBackToVariableName)
 {
 	if (SelectedTreeNodes.Num() > 0)
 	{
@@ -2598,6 +2598,14 @@ void SSubobjectEditor::RestoreSelectionState(TArray<FSubobjectEditorTreeNodePtrT
 			else
 			{
 				FSubobjectEditorTreeNodePtrType NodeToSelectPtr = FindSlateNodeForHandle(SelectedTreeNodes[i]->GetDataHandle());
+
+				// If we didn't find something for this exact handle, fall back to just search for something
+				// with the same variable name. This helps to still preserve selection across re-compiles of a class.
+				if (!NodeToSelectPtr.IsValid() && bFallBackToVariableName)
+				{
+					NodeToSelectPtr = FindSlateNodeForVariableName(SelectedTreeNodes[i]->GetVariableName());
+				}
+
 				if (NodeToSelectPtr.IsValid())
 				{
 					TreeWidget->SetItemSelection(NodeToSelectPtr, true, SelectInfo);
@@ -2987,6 +2995,29 @@ FSubobjectEditorTreeNodePtrType SSubobjectEditor::FindSlateNodeForHandle(const F
 	return OutNodePtr;
 }
 
+FSubobjectEditorTreeNodePtrType SSubobjectEditor::FindSlateNodeForVariableName(FName InVariableName) const
+{
+	FSubobjectEditorTreeNodePtrType OutNodePtr;
+	if (RootNodes.Num() > 0)
+	{
+		TSet<FSubobjectEditorTreeNodePtrType> VisitedNodes;
+		DepthFirstTraversal(RootNodes[0], VisitedNodes,
+			[&OutNodePtr, InVariableName](
+			const FSubobjectEditorTreeNodePtrType& CurNodePtr)
+			{
+				if (CurNodePtr->GetDataHandle().IsValid())
+				{
+					if (CurNodePtr->GetDataHandle().GetSharedDataPtr()->GetVariableName() == InVariableName)
+					{
+						OutNodePtr = CurNodePtr;
+					}
+				}
+			});
+	}
+
+	return OutNodePtr;
+}
+
 FSubobjectEditorTreeNodePtrType SSubobjectEditor::FindSlateNodeForObject(const UObject* InObject, bool bIncludeAttachmentComponents) const
 {
 	FSubobjectEditorTreeNodePtrType OutNodePtr;
diff --git a/Engine/Source/Editor/SubobjectEditor/Public/SSubobjectEditor.h b/Engine/Source/Editor/SubobjectEditor/Public/SSubobjectEditor.h
--- a/Engine/Source/Editor/SubobjectEditor/Public/SSubobjectEditor.h
+++ b/Engine/Source/Editor/SubobjectEditor/Public/SSubobjectEditor.h
@@ -555,6 +555,9 @@ public:
 	// Attempt to find an existing slate node that matches the given handle
 	FSubobjectEditorTreeNodePtrType FindSlateNodeForHandle(const FSubobjectDataHandle& Handle, FSubobjectEditorTreeNodePtrType InStartNodePtr = FSubobjectEditorTreeNodePtrType()) const;
 
+	// Attempt to find an existing slate node that has a given variable name
+	FSubobjectEditorTreeNodePtrType FindSlateNodeForVariableName(FName InVariableName) const;
+
 	/** Pointer to the current object that is represented by the subobject editor */
 	UObject* GetObjectContext() const;
 
@@ -567,7 +570,7 @@ public:
 protected:
 
 	/** Restore the previous selection state when updating the tree */
-	virtual void RestoreSelectionState(TArray<FSubobjectEditorTreeNodePtrType>& SelectedTreeNodes);
+	virtual void RestoreSelectionState(TArray<FSubobjectEditorTreeNodePtrType>& SelectedTreeNodes, bool bFallBackToVariableName = true);
 
 	/** If true, then the blueprint should be modified on TryHandleAssetDragDropOperation */
 	virtual bool ShouldModifyBPOnAssetDrop() const { return false; }
