Made the subobject editor respect overrides Interface API

diff --git a/Engine/Source/Editor/Kismet/Private/SSCSEditor.cpp b/Engine/Source/Editor/Kismet/Private/SSCSEditor.cpp
--- a/Engine/Source/Editor/Kismet/Private/SSCSEditor.cpp
+++ b/Engine/Source/Editor/Kismet/Private/SSCSEditor.cpp
@@ -4049,7 +4049,7 @@ void SSCSEditor::Construct( const FArguments& InArgs )
 	[
 		SNew(SComponentClassCombo)
 		.AddMetaData<FTagMetaData>(FTagMetaData(TEXT("Actor.AddComponent")))
-		.Visibility(HideComponentClassCombo.Get() ? EVisibility::Hidden : EVisibility::Visible)
+		.Visibility(this, &SSCSEditor::GetComponentClassComboButtonVisibility)
 		.OnComponentClassSelected(this, &SSCSEditor::PerformComboAddClass)
 		.ToolTipText(LOCTEXT("AddComponent_Tooltip", "Adds a new component to this actor"))
 		.IsEnabled(AllowEditing)
diff --git a/Engine/Source/Editor/SubobjectEditor/Private/SSubobjectEditor.cpp b/Engine/Source/Editor/SubobjectEditor/Private/SSubobjectEditor.cpp
--- a/Engine/Source/Editor/SubobjectEditor/Private/SSubobjectEditor.cpp
+++ b/Engine/Source/Editor/SubobjectEditor/Private/SSubobjectEditor.cpp
@@ -2522,8 +2522,13 @@ void SSubobjectEditor::UpdateTree(bool bRegenerateTreeNodes /* = true */)
 					ParentNode->AddChild(NewNode);
 					TreeWidget->SetItemExpansion(ParentNode, true);
 
-					// Add a seperator after the default scene root, but only if there are more items below it
-					if (Data->IsDefaultSceneRoot() && Data->IsInheritedComponent() && SubobjectData.Find(Handle) < SubobjectData.Num() - 1)
+					const bool bFilteredOut = RefreshFilteredState(NewNode, false);
+
+					// Add a separator after the default scene root, but only it is not filtered out and if there are more items below it
+					if (!bFilteredOut &&
+						Data->IsDefaultSceneRoot() &&
+						Data->IsInheritedComponent() && 
+						SubobjectData.Find(Handle) < SubobjectData.Num() - 1)
 					{
 						SeperatorNode = MakeShareable<FSubobjectEditorTreeNode>(
 							new FSubobjectEditorTreeNode(FSubobjectDataHandle::InvalidHandle, /** bIsSeperator */true));
@@ -3602,4 +3607,4 @@ TArray<UObject*> USubobjectEditorMenuContext::GetSelectedObjects() const
 	return Result;
 }
 
-#undef LOCTEXT_NAMESPACE
+#undef LOCTEXT_NAMESPACE
\ No newline at end of file
diff --git a/Engine/Source/Editor/SubobjectEditor/Private/SSubobjectInstanceEditor.cpp b/Engine/Source/Editor/SubobjectEditor/Private/SSubobjectInstanceEditor.cpp
--- a/Engine/Source/Editor/SubobjectEditor/Private/SSubobjectInstanceEditor.cpp
+++ b/Engine/Source/Editor/SubobjectEditor/Private/SSubobjectInstanceEditor.cpp
@@ -76,7 +76,7 @@ void SSubobjectInstanceEditor::Construct(const FArguments& InArgs)
 	[
 		SNew(SComponentClassCombo)
 		.AddMetaData<FTagMetaData>(FTagMetaData(TEXT("Actor.AddComponent")))
-		.Visibility(HideComponentClassCombo.Get() ? EVisibility::Hidden : EVisibility::Visible)
+		.Visibility(this, &SSubobjectInstanceEditor::GetComponentClassComboButtonVisibility)
 		.OnSubobjectClassSelected(this, &SSubobjectInstanceEditor::PerformComboAddClass)
 		.ToolTipText(LOCTEXT("AddComponent_Tooltip", "Adds a new component to this actor"))
 		.IsEnabled(true)
