Many fixes of entity prefab workflow breaks

diff --git a/Engine/Source/Editor/SubobjectEditor/Private/SSubobjectBlueprintEditor.cpp b/Engine/Source/Editor/SubobjectEditor/Private/SSubobjectBlueprintEditor.cpp
--- a/Engine/Source/Editor/SubobjectEditor/Private/SSubobjectBlueprintEditor.cpp
+++ b/Engine/Source/Editor/SubobjectEditor/Private/SSubobjectBlueprintEditor.cpp
@@ -189,10 +189,9 @@ void SSubobjectBlueprintEditor::OnDeleteNodes()
 	GUnrealEd->ComponentVisManager.ClearActiveComponentVis();
 
 	UBlueprint* Blueprint = GetBlueprint();
-	check(Blueprint != nullptr);
 	
 	// Get the current render info for the blueprint. If this is NULL then the blueprint is not currently visualizable (no visible primitive components)
-	FThumbnailRenderingInfo* RenderInfo = GUnrealEd->GetThumbnailManager()->GetRenderingInfo( Blueprint );
+	FThumbnailRenderingInfo* RenderInfo = Blueprint ? GUnrealEd->GetThumbnailManager()->GetRenderingInfo( Blueprint ) : nullptr;
 
 	// A lamda for displaying a confirm message to the user if there is a dynamic delegate bound to the 
 	// component they are trying to delete
@@ -260,13 +259,16 @@ void SSubobjectBlueprintEditor::OnDeleteNodes()
 
 			// If we had a thumbnail before we deleted any components, check to see if we should clear it
 			// If we deleted the final visualizable primitive from the blueprint, GetRenderingInfo should return NULL
-			FThumbnailRenderingInfo* NewRenderInfo = GUnrealEd->GetThumbnailManager()->GetRenderingInfo(Blueprint);
-			if (RenderInfo && !NewRenderInfo)
+			if (Blueprint && RenderInfo)
 			{
-				// We removed the last visible primitive component, clear the thumbnail
-				const FString BPFullName = FString::Printf(TEXT("%s %s"), *Blueprint->GetClass()->GetName(), *Blueprint->GetPathName());
-				UPackage* BPPackage = Blueprint->GetOutermost();
-				ThumbnailTools::CacheEmptyThumbnail( BPFullName, BPPackage );
+				FThumbnailRenderingInfo* NewRenderInfo = GUnrealEd->GetThumbnailManager()->GetRenderingInfo(Blueprint);
+				if (RenderInfo && !NewRenderInfo)
+				{
+					// We removed the last visible primitive component, clear the thumbnail
+					const FString BPFullName = FString::Printf(TEXT("%s %s"), *Blueprint->GetClass()->GetName(), *Blueprint->GetPathName());
+					UPackage* BPPackage = Blueprint->GetOutermost();
+					ThumbnailTools::CacheEmptyThumbnail( BPFullName, BPPackage );
+				}
 			}
 			// Do this AFTER marking the Blueprint as modified
 			UpdateSelectionFromNodes(TreeWidget->GetSelectedItems());
diff --git a/Engine/Source/Runtime/Engine/Classes/Components/SceneComponent.h b/Engine/Source/Runtime/Engine/Classes/Components/SceneComponent.h
--- a/Engine/Source/Runtime/Engine/Classes/Components/SceneComponent.h
+++ b/Engine/Source/Runtime/Engine/Classes/Components/SceneComponent.h
@@ -719,7 +719,7 @@ public:
 	* @param  SocketName			Optional socket to attach to on the parent.
 	* @return True if attachment is successful (or already attached to requested parent/socket), false if attachment is rejected and there is no change in AttachParent.
 	*/
-	ENGINE_API bool AttachToComponent(USceneComponent* InParent, const FAttachmentTransformRules& AttachmentRules, FName InSocketName = NAME_None );
+	ENGINE_API virtual bool AttachToComponent(USceneComponent* InParent, const FAttachmentTransformRules& AttachmentRules, FName InSocketName = NAME_None );
 
 	/**
 	* Attach this component to another scene component, optionally at a named socket. It is valid to call this on components whether or not they have been Registered.
