Speculative fix for an crash in the function SSubobjectEditor::RestoreSelectionState.
If the node doesn't have a valid handle for its data source it will crash when calling the function GetVariableName.

diff --git a/Engine/Source/Editor/SubobjectEditor/Private/SSubobjectEditor.cpp b/Engine/Source/Editor/SubobjectEditor/Private/SSubobjectEditor.cpp
--- a/Engine/Source/Editor/SubobjectEditor/Private/SSubobjectEditor.cpp
+++ b/Engine/Source/Editor/SubobjectEditor/Private/SSubobjectEditor.cpp
@@ -2602,13 +2602,14 @@ void SSubobjectEditor::RestoreSelectionState(TArray<FSubobjectEditorTreeNodePtrT
 			}
 			else
 			{
-				FSubobjectEditorTreeNodePtrType NodeToSelectPtr = FindSlateNodeForHandle(SelectedTreeNodes[i]->GetDataHandle());
+				FSubobjectDataHandle CurrentNodeDataHandle = SelectedTreeNodes[i]->GetDataHandle();
+				FSubobjectEditorTreeNodePtrType NodeToSelectPtr = FindSlateNodeForHandle(CurrentNodeDataHandle);
 
 				// If we didn't find something for this exact handle, fall back to just search for something
 				// with the same variable name. This helps to still preserve selection across re-compiles of a class.
-				if (!NodeToSelectPtr.IsValid() && bFallBackToVariableName)
+				if (!NodeToSelectPtr.IsValid() && CurrentNodeDataHandle.IsValid() && bFallBackToVariableName)
 				{
-					NodeToSelectPtr = FindSlateNodeForVariableName(SelectedTreeNodes[i]->GetVariableName());
+					NodeToSelectPtr = FindSlateNodeForVariableName(CurrentNodeDataHandle.GetSharedDataPtr()->GetVariableName());
 				}
 
 				if (NodeToSelectPtr.IsValid())
