Fix an incorrect tooltip on the Default Scene Root component in the subobject editor. It incorrectly said that "it cannot be copied", but it is perfectly valid to copy the scene root component and paste it somewhere else in the subobject hierarchy.

diff --git a/Engine/Source/Editor/SubobjectEditor/Private/SSubobjectEditor.cpp b/Engine/Source/Editor/SubobjectEditor/Private/SSubobjectEditor.cpp
--- a/Engine/Source/Editor/SubobjectEditor/Private/SSubobjectEditor.cpp
+++ b/Engine/Source/Editor/SubobjectEditor/Private/SSubobjectEditor.cpp
@@ -1634,12 +1634,12 @@ FText SSubobject_RowWidget::GetTooltipText() const
 		if (Data->IsInheritedComponent())
 		{
 			return LOCTEXT("InheritedDefaultSceneRootToolTip",
-			               "This is the default scene root component. It cannot be copied, renamed or deleted.\nIt has been inherited from the parent class, so its properties cannot be edited here.\nNew scene components will automatically be attached to it.");
+			               "This is the default scene root component. It cannot be renamed or deleted.\nIt has been inherited from the parent class, so its properties cannot be edited here.\nNew scene components will automatically be attached to it.");
 		}
 		else
 		{
 			return LOCTEXT("DefaultSceneRootToolTip",
-			               "This is the default scene root component. It cannot be copied, renamed or deleted.\nIt can be replaced by drag/dropping another scene component over it.");
+			               "This is the default scene root component. It cannot be renamed or deleted.\nIt can be replaced by drag/dropping another scene component over it.");
 		}
 	}
 	else
