Fix for crash when adding a component in the level editor after its scene root has been deleted, which is permitted for actors based on native classes that do not have a required scene root

diff --git a/Engine/Source/Editor/SubobjectDataInterface/Private/SubobjectDataSubsystem.cpp b/Engine/Source/Editor/SubobjectDataInterface/Private/SubobjectDataSubsystem.cpp
--- a/Engine/Source/Editor/SubobjectDataInterface/Private/SubobjectDataSubsystem.cpp
+++ b/Engine/Source/Editor/SubobjectDataInterface/Private/SubobjectDataSubsystem.cpp
@@ -1000,7 +1000,8 @@ FSubobjectDataHandle USubobjectDataSubsystem::AddNewSubobject(const FAddNewSubob
 				{
 					// If this is an actor with no subobjects on it, then set the new subobject as scene root of the actor.
 					// This can occur if the user has placed in a native C++ class to the world with no subobjects on it
-					if (ParentObjData->IsActor() && ParentObjData->GetChildrenHandles().IsEmpty())
+					if ((ParentObjData->IsActor() && ParentObjData->GetChildrenHandles().IsEmpty()) || 
+						ActorInstance->GetRootComponent() == nullptr)
 					{
 						ActorInstance->SetRootComponent(NewSceneComponent);
 					}
diff --git a/Engine/Source/Editor/SubobjectEditor/Private/SSubobjectEditor.cpp b/Engine/Source/Editor/SubobjectEditor/Private/SSubobjectEditor.cpp
--- a/Engine/Source/Editor/SubobjectEditor/Private/SSubobjectEditor.cpp
+++ b/Engine/Source/Editor/SubobjectEditor/Private/SSubobjectEditor.cpp
@@ -1638,8 +1638,16 @@ FText SSubobject_RowWidget::GetTooltipText() const
 		}
 		else
 		{
-			return LOCTEXT("DefaultSceneRootToolTip",
-			               "This is the default scene root component. It cannot be renamed or deleted.\nIt can be replaced by drag/dropping another scene component over it.");
+			if (Data->CanDelete())
+			{
+				return LOCTEXT("DefaultSceneRootDeletableToolTip",
+					"This is the default scene root component.\nIt can be replaced by drag/dropping another scene component over it.");
+			}
+			else
+			{
+				return LOCTEXT("DefaultSceneRootToolTip",
+					"This is the default scene root component. It cannot be renamed or deleted.\nIt can be replaced by drag/dropping another scene component over it.");
+			}
 		}
 	}
 	else
