Updated headers using UnrealCodeFixup to make sure dllstorage is on methods/staticvars instead of types. Part 5/n

diff --git a/Engine/Source/Editor/StructUtilsEditor/Public/PropertyBagDetails.h b/Engine/Source/Editor/StructUtilsEditor/Public/PropertyBagDetails.h
--- a/Engine/Source/Editor/StructUtilsEditor/Public/PropertyBagDetails.h
+++ b/Engine/Source/Editor/StructUtilsEditor/Public/PropertyBagDetails.h
@@ -7,6 +7,8 @@
 #include "StructUtils/PropertyBag.h"
 #include "PropertyBagDetails.generated.h"
 
+#define UE_API STRUCTUTILSEDITOR_API
+
 class FReply;
 enum class ESelectorType : uint8;
 struct FEdGraphSchemaAction;
@@ -76,7 +78,7 @@ namespace UE::StructUtils
 /**
  * Type customization for FInstancedPropertyBag.
  */
-class STRUCTUTILSEDITOR_API FPropertyBagDetails : public IPropertyTypeCustomization
+class FPropertyBagDetails : public IPropertyTypeCustomization
 {
 public:
 	// Required for deprecated variable 'bFixedLayout'. TODO: Remove when removing 5.6 deprecations.
@@ -87,18 +89,18 @@ public:
 	FPropertyBagDetails(FPropertyBagDetails&& Other) noexcept = delete;
 	FPropertyBagDetails& operator=(const FPropertyBagDetails& Other) = delete;
 	FPropertyBagDetails& operator=(FPropertyBagDetails&& Other) noexcept = delete;
-	PRAGMA_ENABLE_DEPRECATION_WARNINGS
+	UE_API PRAGMA_ENABLE_DEPRECATION_WARNINGS
 
 	/** Makes a new instance of this detail layout class for a specific detail view requesting it */
 	static TSharedRef<IPropertyTypeCustomization> MakeInstance();
 
 	/** Creates add property widget. */
-	static TSharedPtr<SWidget> MakeAddPropertyWidget(TSharedPtr<IPropertyHandle> InStructProperty, TSharedPtr<IPropertyUtilities> InPropUtils, EPropertyBagPropertyType DefaultType = EPropertyBagPropertyType::Bool, FSlateColor IconColor = FSlateColor::UseForeground());
+	static UE_API TSharedPtr<SWidget> MakeAddPropertyWidget(TSharedPtr<IPropertyHandle> InStructProperty, TSharedPtr<IPropertyUtilities> InPropUtils, EPropertyBagPropertyType DefaultType = EPropertyBagPropertyType::Bool, FSlateColor IconColor = FSlateColor::UseForeground());
 
 protected:
 	/** IPropertyTypeCustomization interface */
-	virtual void CustomizeHeader(TSharedRef<IPropertyHandle> StructPropertyHandle, FDetailWidgetRow& HeaderRow, IPropertyTypeCustomizationUtils& StructCustomizationUtils) override;
-	virtual void CustomizeChildren(TSharedRef<IPropertyHandle> StructPropertyHandle, IDetailChildrenBuilder& StructBuilder, IPropertyTypeCustomizationUtils& StructCustomizationUtils) override;
+	UE_API virtual void CustomizeHeader(TSharedRef<IPropertyHandle> StructPropertyHandle, FDetailWidgetRow& HeaderRow, IPropertyTypeCustomizationUtils& StructCustomizationUtils) override;
+	UE_API virtual void CustomizeChildren(TSharedRef<IPropertyHandle> StructPropertyHandle, IDetailChildrenBuilder& StructBuilder, IPropertyTypeCustomizationUtils& StructCustomizationUtils) override;
 
 private:
 	/** Handle to the struct property being edited */
@@ -120,7 +122,7 @@ private:
  * Use FPropertyBagDetails::MakeAddPropertyWidget() to create the add property widget.
  * OnChildRowAdded() is called right after each property is added, which allows the property row to be customizable.
  */
-class STRUCTUTILSEDITOR_API FPropertyBagInstanceDataDetails : public FInstancedStructDataDetails
+class FPropertyBagInstanceDataDetails : public FInstancedStructDataDetails
 {
 public:
 	struct FConstructParams
@@ -131,8 +133,8 @@ public:
 		EPropertyBagChildRowFeatures ChildRowFeatures = EPropertyBagChildRowFeatures::Default;
 	};
 
-	explicit FPropertyBagInstanceDataDetails(const FConstructParams& ConstructParams);
-	FPropertyBagInstanceDataDetails(TSharedPtr<IPropertyHandle> InStructProperty, const TSharedPtr<IPropertyUtilities>& InPropUtils, bool bInFixedLayout, bool bInAllowContainers = true);
+	UE_API explicit FPropertyBagInstanceDataDetails(const FConstructParams& ConstructParams);
+	UE_API FPropertyBagInstanceDataDetails(TSharedPtr<IPropertyHandle> InStructProperty, const TSharedPtr<IPropertyUtilities>& InPropUtils, bool bInFixedLayout, bool bInAllowContainers = true);
 
 	// Required for deprecated variable 'bFixedLayout' and 'bAllowArrays'. TODO: Remove when removing 5.6 deprecations.
 	PRAGMA_DISABLE_DEPRECATION_WARNINGS
@@ -141,10 +143,10 @@ public:
 	FPropertyBagInstanceDataDetails(FPropertyBagInstanceDataDetails&& Other) noexcept = delete;
 	FPropertyBagInstanceDataDetails& operator=(const FPropertyBagInstanceDataDetails& Other) = delete;
 	FPropertyBagInstanceDataDetails& operator=(FPropertyBagInstanceDataDetails&& Other) noexcept = delete;
-	PRAGMA_ENABLE_DEPRECATION_WARNINGS
+	UE_API PRAGMA_ENABLE_DEPRECATION_WARNINGS
 
 	virtual void OnGroupRowAdded(IDetailGroup& GroupRow, int32 Level, const FString& Category) const override;
-	virtual void OnChildRowAdded(IDetailPropertyRow& ChildRow) override;
+	UE_API virtual void OnChildRowAdded(IDetailPropertyRow& ChildRow) override;
 	
 	/** Enum describing if a property is overridden, or undetermined (e.g. multiselection) */
 	enum class EPropertyOverrideState
@@ -209,20 +211,20 @@ public:
 	}
 
 	/** @return true if property of specified child property is overridden. */
-	virtual EPropertyOverrideState IsPropertyOverridden(TSharedPtr<IPropertyHandle> ChildPropertyHandle) const;
+	UE_API virtual EPropertyOverrideState IsPropertyOverridden(TSharedPtr<IPropertyHandle> ChildPropertyHandle) const;
 
 	/** Called to set the override state of specified child property. */
-	virtual void SetPropertyOverride(TSharedPtr<IPropertyHandle> ChildPropertyHandle, const bool bIsOverridden);
+	UE_API virtual void SetPropertyOverride(TSharedPtr<IPropertyHandle> ChildPropertyHandle, const bool bIsOverridden);
 
 	/** @return true if the child property has default value. */
-	virtual bool IsDefaultValue(TSharedPtr<IPropertyHandle> ChildPropertyHandle) const;
+	UE_API virtual bool IsDefaultValue(TSharedPtr<IPropertyHandle> ChildPropertyHandle) const;
 
 	/** Called to reset the child property to default value. */
-	virtual void ResetToDefault(TSharedPtr<IPropertyHandle> ChildPropertyHandle);
+	UE_API virtual void ResetToDefault(TSharedPtr<IPropertyHandle> ChildPropertyHandle);
 
 protected:
 	UE_DEPRECATED(5.6, "Replaced by 'CreatePropertyDetailsWidget' to allow for more customizable drop-down menu location and content.")
-	TSharedRef<SWidget> OnPropertyNameContent(TSharedPtr<IPropertyHandle> ChildPropertyHandle, TSharedPtr<SInlineEditableTextBlock> InlineWidget) const;
+	UE_API TSharedRef<SWidget> OnPropertyNameContent(TSharedPtr<IPropertyHandle> ChildPropertyHandle, TSharedPtr<SInlineEditableTextBlock> InlineWidget) const;
 
 	TSharedPtr<IPropertyHandle> BagStructProperty;
 	TSharedPtr<IPropertyUtilities> PropUtils;
@@ -239,10 +241,12 @@ protected:
 /**
  * Specific property bag schema to allow customizing the requirements (e.g. supported containers).
  */
-UCLASS()
-class STRUCTUTILSEDITOR_API UPropertyBagSchema : public UEdGraphSchema_K2
+UCLASS(MinimalAPI)
+class UPropertyBagSchema : public UEdGraphSchema_K2
 {
 	GENERATED_BODY()
 public:
-	virtual bool SupportsPinTypeContainer(TWeakPtr<const FEdGraphSchemaAction> SchemaAction, const FEdGraphPinType& PinType, const EPinContainerType& ContainerType) const override;
+	UE_API virtual bool SupportsPinTypeContainer(TWeakPtr<const FEdGraphSchemaAction> SchemaAction, const FEdGraphPinType& PinType, const EPinContainerType& ContainerType) const override;
 };
+
+#undef UE_API
diff --git a/Engine/Source/Editor/StructUtilsEditor/Public/PropertyBagDragDropHandler.h b/Engine/Source/Editor/StructUtilsEditor/Public/PropertyBagDragDropHandler.h
--- a/Engine/Source/Editor/StructUtilsEditor/Public/PropertyBagDragDropHandler.h
+++ b/Engine/Source/Editor/StructUtilsEditor/Public/PropertyBagDragDropHandler.h
@@ -6,6 +6,8 @@
 #include "IDetailDragDropHandler.h"
 #include "StructUtils/PropertyBag.h"
 
+#define UE_API STRUCTUTILSEDITOR_API
+
 class FPropertyBagDetailsDragDropOp;
 class IPropertyBagEdGraphDragAndDrop;
 
@@ -24,7 +26,7 @@ enum class EPropertyBagDropState
  * Provides information about the source row (single property) being dragged.
  * Inherits from FGraphEditorDragDropAction to support dragging to a graph pin/node/panel.
  */
-class STRUCTUTILSEDITOR_API FPropertyBagDetailsDragDropOp : public FGraphEditorDragDropAction
+class FPropertyBagDetailsDragDropOp : public FGraphEditorDragDropAction
 {
 public:
 	DRAG_DROP_OPERATOR_TYPE(FPropertyBagDetailsDragDropOp, FGraphEditorDragDropAction);
@@ -42,23 +44,23 @@ public:
 	 * Property Bag Details - Drop Operation Constructor
 	 * @param InPropertyDesc Property description for the dropped property.
 	 */
-	FPropertyBagDetailsDragDropOp(const FPropertyBagPropertyDesc& InPropertyDesc);
+	UE_API FPropertyBagDetailsDragDropOp(const FPropertyBagPropertyDesc& InPropertyDesc);
 
 	/**
 	 * Sets the pop up widget's icon and text.
 	 * @param NewDropState The state of the drop operation.
 	 * @param OverriddenDecoration An optional overridden decorator for the UI, including Message, Icon, and IconColor
 	 */
-	void SetDecoration(EPropertyBagDropState NewDropState, TOptional<FDecoration> OverriddenDecoration = {});
+	UE_API void SetDecoration(EPropertyBagDropState NewDropState, TOptional<FDecoration> OverriddenDecoration = {});
 
 	// ~Begin FGraphEditorDragDropAction interface
 	/** Event when the drag and drop operation changes hover targets. */
-	virtual void HoverTargetChanged() override;
+	UE_API virtual void HoverTargetChanged() override;
 
 	/** Visibility of the cursor decorator icon. */
-	virtual EVisibility GetIconVisible() const override;
+	UE_API virtual EVisibility GetIconVisible() const override;
 	/** Visibility of the cursor decorator error icon. */
-	virtual EVisibility GetErrorIconVisible() const override;
+	UE_API virtual EVisibility GetErrorIconVisible() const override;
 
 	/**
 	 * The drop operation has been executed on an UEdGraphPin in an Editor Graph.
@@ -66,7 +68,7 @@ public:
 	 * @param GraphPosition Local 2D graph coordinates.
 	 * @return If the drop operation was handled by the graph.
 	 */
-	virtual FReply DroppedOnPin(const FVector2f& ScreenPosition, const FVector2f& GraphPosition) override;
+	UE_API virtual FReply DroppedOnPin(const FVector2f& ScreenPosition, const FVector2f& GraphPosition) override;
 
 	/**
 	 * The drop operation has been executed on an UEdGraphNode in an Editor Graph.
@@ -74,7 +76,7 @@ public:
 	 * @param GraphPosition Local 2D graph coordinates.
 	 * @return If the drop operation was handled by the graph.
 	 */
-	virtual FReply DroppedOnNode(const FVector2f& ScreenPosition, const FVector2f& GraphPosition) override;
+	UE_API virtual FReply DroppedOnNode(const FVector2f& ScreenPosition, const FVector2f& GraphPosition) override;
 
 	/**
 	 * The drop operation has been executed on an SGraphPanel in an Editor Graph.
@@ -84,7 +86,7 @@ public:
 	 * @param Graph The graph owner of the target panel.
 	 * @return If the drop operation was handled by the graph.
 	 */
-	virtual FReply DroppedOnPanel(const TSharedRef<SWidget>& Panel, const FVector2f& ScreenPosition, const FVector2f& GraphPosition, UEdGraph& Graph) override;
+	UE_API virtual FReply DroppedOnPanel(const TSharedRef<SWidget>& Panel, const FVector2f& ScreenPosition, const FVector2f& GraphPosition, UEdGraph& Graph) override;
 	// ~End FGraphEditorDragDropAction interface
 
 	/** The property description of the dragged property. */
@@ -93,11 +95,11 @@ public:
 	EPropertyBagDropState CurrentDropState = EPropertyBagDropState::Invalid;
 
 private:
-	IPropertyBagEdGraphDragAndDrop* GetPropertyBagEdGraphDragAndDropInterface() const;
+	UE_API IPropertyBagEdGraphDragAndDrop* GetPropertyBagEdGraphDragAndDropInterface() const;
 };
 
 /** Handles drag-and-drop (as a target) for a single property's child widget row. */
-class STRUCTUTILSEDITOR_API FPropertyBagDetailsDragDropHandlerTarget : public IDetailDragDropHandler
+class FPropertyBagDetailsDragDropHandlerTarget : public IDetailDragDropHandler
 {
 public:
 	/**
@@ -110,22 +112,22 @@ public:
 	 * @param CanAcceptDragDrop [Signature]: TOptional<EItemDropZone>(TSharedPtr<FPropertyBagDetailsDragDropOp>, EItemDropZone)
 	 * @param OnDragDrop [Signature]: FReply(const FPropertyBagPropertyDesc&, EItemDropZone)
 	 */
-	FPropertyBagDetailsDragDropHandlerTarget(const FCanAcceptPropertyBagDetailsRowDropOp& CanAcceptDragDrop, const FOnPropertyBagDetailsRowDropOp& OnDragDrop);
+	UE_API FPropertyBagDetailsDragDropHandlerTarget(const FCanAcceptPropertyBagDetailsRowDropOp& CanAcceptDragDrop, const FOnPropertyBagDetailsRowDropOp& OnDragDrop);
 
 	/** Creates the drag and drop operation. Disabled for the target. Enabled in subclasses. */
-	virtual TSharedPtr<FDragDropOperation> CreateDragDropOperation() const override;
+	UE_API virtual TSharedPtr<FDragDropOperation> CreateDragDropOperation() const override;
 
 	/**
 	 * Bind the delegate for determining if this handler can accept this drag and drop operation.
 	 * @param CanAcceptDragDrop [Signature]: TOptional<EItemDropZone>(TSharedPtr<FPropertyBagDetailsDragDropOp>, EItemDropZone)
 	 */
-	virtual void BindCanAcceptDragDrop(FCanAcceptPropertyBagDetailsRowDropOp&& CanAcceptDragDrop);
+	UE_API virtual void BindCanAcceptDragDrop(FCanAcceptPropertyBagDetailsRowDropOp&& CanAcceptDragDrop);
 
 	/**
 	 * Bind the delegate to handle a valid drag and drop operation.
 	 * @param OnDragDrop [Signature]: FReply(const FPropertyBagPropertyDesc&, EItemDropZone)
 	 */
-	virtual void BindOnHandleDragDrop(FOnPropertyBagDetailsRowDropOp&& OnDragDrop);
+	UE_API virtual void BindOnHandleDragDrop(FOnPropertyBagDetailsRowDropOp&& OnDragDrop);
 
 	/** Disable automatic creation of the handle widget for targets. It is overridden in the source below. */
 	virtual bool UseHandleWidget() const override
@@ -140,7 +142,7 @@ protected:
 	 * @param DropZone Zone of the drop, relative to this child row.
 	 * @return [Optional] Target zone, which can be overridden. False if the row will not accept.
 	 */
-	virtual TOptional<EItemDropZone> CanAcceptDrop(const FDragDropEvent& DragDropSource, EItemDropZone DropZone) const override;
+	UE_API virtual TOptional<EItemDropZone> CanAcceptDrop(const FDragDropEvent& DragDropSource, EItemDropZone DropZone) const override;
 
 	/**
 	 * Accept and handle the drag and drop operation.
@@ -148,7 +150,7 @@ protected:
 	 * @param DropZone Zone of the drop, relative to this child row.
 	 * @return True if the operation was successful.
 	 */
-	virtual bool AcceptDrop(const FDragDropEvent& DragDropSource, EItemDropZone DropZone) const override;
+	UE_API virtual bool AcceptDrop(const FDragDropEvent& DragDropSource, EItemDropZone DropZone) const override;
 
 private:
 	/**
@@ -165,21 +167,21 @@ private:
 };
 
 /** Handles drag-and-drop (as a source or target) for a single property's child widget row. */
-class STRUCTUTILSEDITOR_API FPropertyBagDetailsDragDropHandler : public FPropertyBagDetailsDragDropHandlerTarget
+class FPropertyBagDetailsDragDropHandler : public FPropertyBagDetailsDragDropHandlerTarget
 {
 public:
 	/**
 	 * Constructor.
 	 * @param InPropertyDesc The property desc of the property represented by this child widget row.
 	 */
-	FPropertyBagDetailsDragDropHandler(const FPropertyBagPropertyDesc& InPropertyDesc);
+	UE_API FPropertyBagDetailsDragDropHandler(const FPropertyBagPropertyDesc& InPropertyDesc);
 
 protected:
 	/**
 	 * Creates the drag and drop operation object that contains necessary information for a drag and drop operation.
 	 * @return Shared pointer to the drag and drop operation.
 	 */
-	virtual TSharedPtr<FDragDropOperation> CreateDragDropOperation() const override;
+	UE_API virtual TSharedPtr<FDragDropOperation> CreateDragDropOperation() const override;
 
 	/** Enable the automatic creation of the handle (grip) widget for the source. */
 	virtual bool UseHandleWidget() const override { return true; }
@@ -188,3 +190,5 @@ private:
 	/** The current child row's property bag property description. */
 	const FPropertyBagPropertyDesc PropertyDesc;
 };
+
+#undef UE_API
diff --git a/Engine/Source/Editor/StructUtilsEditor/Public/StructUtilsEditorModule.h b/Engine/Source/Editor/StructUtilsEditor/Public/StructUtilsEditorModule.h
--- a/Engine/Source/Editor/StructUtilsEditor/Public/StructUtilsEditorModule.h
+++ b/Engine/Source/Editor/StructUtilsEditor/Public/StructUtilsEditorModule.h
@@ -5,6 +5,8 @@
 #include "Modules/ModuleInterface.h"
 #include "Kismet2/StructureEditorUtils.h"
 
+#define UE_API STRUCTUTILSEDITOR_API
+
 class IStructUtilsEditor;
 struct FGraphPanelNodeFactory;
 class UUserDefinedStruct;
@@ -12,17 +14,19 @@ class UUserDefinedStruct;
 /**
 * The public interface to this module
 */
-class STRUCTUTILSEDITOR_API FStructUtilsEditorModule : public IModuleInterface, public FStructureEditorUtils::INotifyOnStructChanged
+class FStructUtilsEditorModule : public IModuleInterface, public FStructureEditorUtils::INotifyOnStructChanged
 {
 public:
 	// Begin IModuleInterface
-	virtual void StartupModule() override;
-	virtual void ShutdownModule() override;
+	UE_API virtual void StartupModule() override;
+	UE_API virtual void ShutdownModule() override;
 
 protected:
 
 	// INotifyOnStructChanged
-	virtual void PreChange(const class UUserDefinedStruct* Struct, FStructureEditorUtils::EStructureEditorChangeInfo Info) override;
-	virtual void PostChange(const class UUserDefinedStruct* Struct, FStructureEditorUtils::EStructureEditorChangeInfo Info) override;
+	UE_API virtual void PreChange(const class UUserDefinedStruct* Struct, FStructureEditorUtils::EStructureEditorChangeInfo Info) override;
+	UE_API virtual void PostChange(const class UUserDefinedStruct* Struct, FStructureEditorUtils::EStructureEditorChangeInfo Info) override;
 	// ~INotifyOnStructChanged
-};
\ No newline at end of file
+};
+
+#undef UE_API
diff --git a/Engine/Source/Editor/StructViewer/Public/StructViewerFilter.h b/Engine/Source/Editor/StructViewer/Public/StructViewerFilter.h
--- a/Engine/Source/Editor/StructViewer/Public/StructViewerFilter.h
+++ b/Engine/Source/Editor/StructViewer/Public/StructViewerFilter.h
@@ -3,6 +3,8 @@
 
 #include "CoreMinimal.h"
 
+#define UE_API STRUCTVIEWER_API
+
 class FStructViewerInitializationOptions;
 
 /** Interface class for creating filters for the Struct Viewer. */
@@ -36,7 +38,7 @@ enum class EStructFilterReturn : uint8
 	Passed, 
 	NoItems,
 };
-class STRUCTVIEWER_API FStructViewerFilterFuncs
+class FStructViewerFilterFuncs
 {
 public:
 	virtual ~FStructViewerFilterFuncs() = default;
@@ -49,7 +51,7 @@ public:
 	 *
 	 * @return					EFilterReturn::Passed if it is a child-of a struct in the set, EFilterReturn::Failed if it is not, EFilterReturn::NoItems if the set is empty.
 	 */
-	virtual EStructFilterReturn IfInChildOfStructsSet(TSet<const UScriptStruct*>& InSet, const UScriptStruct* InStruct);
+	UE_API virtual EStructFilterReturn IfInChildOfStructsSet(TSet<const UScriptStruct*>& InSet, const UScriptStruct* InStruct);
 
 	/** 
 	 * Checks if the given struct is a child-of ALL of the classes in a set.
@@ -59,7 +61,7 @@ public:
 	 *
 	 * @return					EFilterReturn::Passed if it is a child-of ALL the structs in the set, EFilterReturn::Failed if it is not, EFilterReturn::NoItems if the set is empty.
 	 */
-	virtual EStructFilterReturn IfMatchesAllInChildOfStructsSet(TSet<const UScriptStruct*>& InSet, const UScriptStruct* InStruct);
+	UE_API virtual EStructFilterReturn IfMatchesAllInChildOfStructsSet(TSet<const UScriptStruct*>& InSet, const UScriptStruct* InStruct);
 
 	/** 
 	 * Checks if the struct is in the structs set.
@@ -69,5 +71,7 @@ public:
 	 *
 	 * @return					EFilterReturn::Passed if the struct is in the structs set, EFilterReturn::Failed if it is not, EFilterReturn::NoItems if the set is empty.
 	 */
-	virtual EStructFilterReturn IfInStructsSet(TSet<const UScriptStruct*>& InSet, const UScriptStruct* InStruct);
+	UE_API virtual EStructFilterReturn IfInStructsSet(TSet<const UScriptStruct*>& InSet, const UScriptStruct* InStruct);
 };
+
+#undef UE_API
diff --git a/Engine/Source/Editor/StructViewer/Public/StructViewerProjectSettings.h b/Engine/Source/Editor/StructViewer/Public/StructViewerProjectSettings.h
--- a/Engine/Source/Editor/StructViewer/Public/StructViewerProjectSettings.h
+++ b/Engine/Source/Editor/StructViewer/Public/StructViewerProjectSettings.h
@@ -19,8 +19,8 @@ class UScriptStruct;
 /**
  * Implements the settings for the Struct Viewer Project Settings
  */
-UCLASS(config=Engine, defaultconfig)
-class STRUCTVIEWER_API UStructViewerProjectSettings : public UObject
+UCLASS(MinimalAPI, config=Engine, defaultconfig)
+class UStructViewerProjectSettings : public UObject
 {
 	GENERATED_UCLASS_BODY()
 
diff --git a/Engine/Source/Editor/SubobjectDataInterface/Public/SubobjectData.h b/Engine/Source/Editor/SubobjectDataInterface/Public/SubobjectData.h
--- a/Engine/Source/Editor/SubobjectDataInterface/Public/SubobjectData.h
+++ b/Engine/Source/Editor/SubobjectDataInterface/Public/SubobjectData.h
@@ -11,6 +11,8 @@
 
 #include "SubobjectData.generated.h"
 
+#define UE_API SUBOBJECTDATAINTERFACE_API
+
 class UBlueprint;
 
 /**
@@ -21,14 +23,14 @@ class UBlueprint;
 * If you wish to modify a subobject, use the SubobjectDataSubsystem.
 */
 USTRUCT(BlueprintType)
-struct SUBOBJECTDATAINTERFACE_API FSubobjectData final
+struct FSubobjectData final
 {
 friend class USubobjectDataSubsystem;
 public:
 	GENERATED_BODY()
 
-	FSubobjectData();
-	FSubobjectData(UObject* ContextObject, const FSubobjectDataHandle& ParentHandle, const bool bIsInheritedSCS);
+	UE_API FSubobjectData();
+	UE_API FSubobjectData(UObject* ContextObject, const FSubobjectDataHandle& ParentHandle, const bool bIsInheritedSCS);
 
 	/** True if this subobject's handle is valid */
 	bool IsValid() const { return Handle.IsValid(); }
@@ -36,34 +38,34 @@ public:
 	/**
 	 * @return Whether or not we can edit properties for this subobject
 	 */
-	bool CanEdit() const;
+	UE_API bool CanEdit() const;
 
 	/**
 	 * @return Whether or not this object represents a subobject that can be deleted
 	 */
-	bool CanDelete() const;
+	UE_API bool CanDelete() const;
 
 	/**
 	* @return Whether or not this object represents a subobject that can be duplicated
 	*/
-	bool CanDuplicate() const;
+	UE_API bool CanDuplicate() const;
 
 	/**
 	* @return Whether or not this object represents a subobject that can be copied
 	*/
-	bool CanCopy() const;
+	UE_API bool CanCopy() const;
 
 	/**
 	 * @return Whether or not this object represents a subobject that can 
 	 * be reparented to other subobjects based on its context.
 	 */
-	bool CanReparent() const;
+	UE_API bool CanReparent() const;
 
 	/**
 	* @return Whether or not this object represents a subobject that can 
 	* be renamed.
 	*/
-	bool CanRename() const;
+	UE_API bool CanRename() const;
 
 	/**
 	* @return	A read-only reference to the object represented by this node.
@@ -76,7 +78,7 @@ public:
      * @note	May not be the same as the value returned by GetObject().
      * @return	A reference to the object represented by this node that can be modified within the given Blueprint.
      */
-	const UObject* GetObjectForBlueprint(UBlueprint* Blueprint) const;
+	UE_API const UObject* GetObjectForBlueprint(UBlueprint* Blueprint) const;
 
 	template<class T>
 	inline const T* GetObjectForBlueprint(UBlueprint* Blueprint) const
@@ -114,54 +116,54 @@ public:
 
 	inline const UActorComponent* FindComponentInstanceInActor(const AActor* InActor) const { return FindMutableComponentInstanceInActor(InActor); }
 
-	UBlueprint* GetBlueprint() const;
-	UBlueprint* GetBlueprintBeingEdited() const;
-
-	bool IsInstancedComponent() const;
-	bool IsInstancedActor() const;
-	bool IsNativeComponent() const;
-	bool IsBlueprintInheritedComponent() const;
-	bool IsInheritedComponent() const;
-	bool IsSceneComponent() const;
-	bool IsRootComponent() const;
-	bool IsDefaultSceneRoot() const;
+	UE_API UBlueprint* GetBlueprint() const;
+	UE_API UBlueprint* GetBlueprintBeingEdited() const;
+
+	UE_API bool IsInstancedComponent() const;
+	UE_API bool IsInstancedActor() const;
+	UE_API bool IsNativeComponent() const;
+	UE_API bool IsBlueprintInheritedComponent() const;
+	UE_API bool IsInheritedComponent() const;
+	UE_API bool IsSceneComponent() const;
+	UE_API bool IsRootComponent() const;
+	UE_API bool IsDefaultSceneRoot() const;
 	/** Returns true if the component template's FName starts with USceneComponent::GetDefaultSceneRootVariableName */
-	bool SceneRootHasDefaultName() const;
+	UE_API bool SceneRootHasDefaultName() const;
 	/* Returns true if this subobject is a component. */
-	bool IsComponent() const;
-	bool IsChildActor() const;
+	UE_API bool IsComponent() const;
+	UE_API bool IsChildActor() const;
 
 	/** Returns true if this subobject is a part of a child actor's hierarchy and not the root actor.  */
-	bool IsChildActorSubtreeObject() const;
-	bool IsRootActor() const;
-	bool IsActor() const;
+	UE_API bool IsChildActorSubtreeObject() const;
+	UE_API bool IsRootActor() const;
+	UE_API bool IsActor() const;
 	/** True if this subobject is an instance of an inherited component (its owner is an instanced actor) */
-	bool IsInstancedInheritedComponent() const;
+	UE_API bool IsInstancedInheritedComponent() const;
 
-	bool IsAttachedTo(const FSubobjectDataHandle& InHandle) const;
+	UE_API bool IsAttachedTo(const FSubobjectDataHandle& InHandle) const;
 
-	FString GetDisplayString(bool bShowNativeComponentNames = true) const;
+	UE_API FString GetDisplayString(bool bShowNativeComponentNames = true) const;
 
 	/** Get the display name of this subobject that should be used during drag/drop operations */
-	FText GetDragDropDisplayText() const;
+	UE_API FText GetDragDropDisplayText() const;
 
 	/**
 	 * Returns a string that gives details about this subobject. For Inherited
 	 * components it will return "(Inherited)". For native components
 	 * this function will return "(NativeComponentName) (Inherited)"
 	 */
-	FText GetDisplayNameContextModifiers(bool bShowNativeComponentNames = true) const;
+	UE_API FText GetDisplayNameContextModifiers(bool bShowNativeComponentNames = true) const;
 	
-	FText GetDisplayName() const;
+	UE_API FText GetDisplayName() const;
 
-	FName GetVariableName() const;
+	UE_API FName GetVariableName() const;
 
 	// Sockets for attaching in the viewport
-	FText GetSocketName() const;
-	FName GetSocketFName() const;
-	bool HasValidSocket() const;
-	void SetSocketName(FName InNewName);
-	void SetupAttachment(FName SocketName, const FSubobjectDataHandle& AttachParentHandle = FSubobjectDataHandle::InvalidHandle);
+	UE_API FText GetSocketName() const;
+	UE_API FName GetSocketFName() const;
+	UE_API bool HasValidSocket() const;
+	UE_API void SetSocketName(FName InNewName);
+	UE_API void SetupAttachment(FName SocketName, const FSubobjectDataHandle& AttachParentHandle = FSubobjectDataHandle::InvalidHandle);
 
 	FSubobjectDataHandle GetHandle() const { return Handle; }
 	FSubobjectDataHandle GetParentHandle() const { return ParentObjectHandle; }
@@ -170,46 +172,46 @@ public:
 	const TArray<FSubobjectDataHandle>& GetChildrenHandles() const { return ChildrenHandles; }
 	
 	/** Walks the parent heirarchy and returns the root subobject in this case */
-	FSubobjectDataHandle GetRootSubobject() const;
+	UE_API FSubobjectDataHandle GetRootSubobject() const;
 
 	// Return true if the given handle is a child of this subobject
-	bool HasChild(const FSubobjectDataHandle& ChildHandle) const;
+	UE_API bool HasChild(const FSubobjectDataHandle& ChildHandle) const;
 
 	// Attempt to find the given handle in out child array and return a pointer to it
-	FSubobjectDataHandle FindChild(const FSubobjectDataHandle& ChildHandle) const;
-	FSubobjectDataHandle FindChildByObject(UObject* ContextObject) const;
+	UE_API FSubobjectDataHandle FindChild(const FSubobjectDataHandle& ChildHandle) const;
+	UE_API FSubobjectDataHandle FindChildByObject(UObject* ContextObject) const;
 
 	/** Get the asset name of this subobject from the asset brokerage */
-	FText GetAssetName() const;
+	UE_API FText GetAssetName() const;
 
 	/** Get the asset path of this subobject from the asset brokerage */
-	FText GetAssetPath() const;
+	UE_API FText GetAssetPath() const;
 
 	/** Returns true if the asset this object represents is visible */
-	bool IsAssetVisible() const;
+	UE_API bool IsAssetVisible() const;
 
 	/**
 	* Retrieves tooltip text describing the specified component's mobility.
 	*
 	* @returns An FText object containing a description of the component's mobility
 	*/
-	FText GetMobilityToolTipText() const;
+	UE_API FText GetMobilityToolTipText() const;
 
 	/**
 	* Retrieves a tooltip text describing if the component is marked Editor only or not
 	*
 	* @returns An FText object containing a description of if the component is marked Editor only or not
 	*/
-	FText GetComponentEditorOnlyTooltipText() const;
+	UE_API FText GetComponentEditorOnlyTooltipText() const;
 
 	/**
 	* Retrieves tooltip text describing where the component was first introduced (for inherited components).
 	* 
 	* @returns An FText object containing a description of when the component was first introduced
 	*/
-	FText GetIntroducedInToolTipText() const;
+	UE_API FText GetIntroducedInToolTipText() const;
 
-	FText GetActorDisplayText() const;
+	UE_API FText GetActorDisplayText() const;
 		
 private:
 	/**
@@ -218,7 +220,7 @@ private:
 	* This only updates this subobject to have the correct handles in it's child structure.
 	* Actual manipulation of subobjects should be done via the SubobjectSubsystem.
 	*/
-	bool AddChildHandleOnly(const FSubobjectDataHandle& InHandle);
+	UE_API bool AddChildHandleOnly(const FSubobjectDataHandle& InHandle);
 
 	/**
 	* Add the given subobject handle as a child of this.
@@ -226,7 +228,7 @@ private:
 	* This only updates this subobject to have the correct handles in it's child structure.
 	* Actual manipulation of subobjects should be done via the SubobjectSubsystem.
 	*/
-	bool RemoveChildHandleOnly(const FSubobjectDataHandle& InHandle);
+	UE_API bool RemoveChildHandleOnly(const FSubobjectDataHandle& InHandle);
 
 	void SetParentHandle(const FSubobjectDataHandle& NewParentHandle) { ParentObjectHandle = NewParentHandle; }
 	void ClearParentHandle() { ParentObjectHandle = FSubobjectDataHandle::InvalidHandle; }
@@ -253,7 +255,7 @@ private:
 
 	// If this subobject is an actor, then return a pointer to that. Otherwise, search for
 	// the owning actor on this component and return it.
-	AActor* GetMutableActorContext();
+	UE_API AActor* GetMutableActorContext();
 
 	template<class T>
     inline T* GetMutableObject(bool bEvenIfPendingKill = false) const
@@ -272,15 +274,15 @@ private:
 		return Cast<UActorComponent>(WeakObjectPtr.Get(bEvenIfPendingKill));
 	}
 
-	UActorComponent* FindMutableComponentInstanceInActor(const AActor* InActor) const;
+	UE_API UActorComponent* FindMutableComponentInstanceInActor(const AActor* InActor) const;
 
 	////////////////////////////////////////////////////
 	// Anything related to an SCS node will be changed with an upcoming refactor and should remain private
 private:
 	// Tries to find a SCS node that was likely responsible for creating the specified instance component.  Note: This is not always possible to do!
-	static USCS_Node* FindSCSNodeForInstance(const UActorComponent* InstanceComponent, UClass* ClassToSearch);
+	static UE_API USCS_Node* FindSCSNodeForInstance(const UActorComponent* InstanceComponent, UClass* ClassToSearch);
 
-	USCS_Node* GetSCSNode(bool bEvenIfPendingKill = false) const;
+	UE_API USCS_Node* GetSCSNode(bool bEvenIfPendingKill = false) const;
 
 	/**
 	* Attempts to set the SCS node member variable based on the WeakObjectPtr.
@@ -291,7 +293,7 @@ private:
 	*
 	* @return True if an SCS node has been set
 	*/
-	bool AttemptToSetSCSNode();
+	UE_API bool AttemptToSetSCSNode();
 
 	// For now, we need a pointer to the SCS node in order to have the correct component template
 	// and to be able to get the UBlueprint* from components. This will only be set upon construction
@@ -302,11 +304,13 @@ private:
 	 * True if this SCS node's blueprint is a child of another Blueprint-generated class,
 	 * which means that we must use the InheritableComponentHandler
 	 */
-	bool IsInheritedSCSNode() const;
-	const UChildActorComponent* GetChildActorComponent(bool bEvenIfPendingKill = false) const;
+	UE_API bool IsInheritedSCSNode() const;
+	UE_API const UChildActorComponent* GetChildActorComponent(bool bEvenIfPendingKill = false) const;
 
 	/** True if this SCS node is inherited from another blueprint generated class. */
 	bool bIsInheritedSubobject = false;
 	bool bIsInheritedSCS = false;
 	bool bIsChildActor = false;
-};
\ No newline at end of file
+};
+
+#undef UE_API
diff --git a/Engine/Source/Editor/SubobjectDataInterface/Public/SubobjectDataBlueprintFunctionLibrary.h b/Engine/Source/Editor/SubobjectDataInterface/Public/SubobjectDataBlueprintFunctionLibrary.h
--- a/Engine/Source/Editor/SubobjectDataInterface/Public/SubobjectDataBlueprintFunctionLibrary.h
+++ b/Engine/Source/Editor/SubobjectDataInterface/Public/SubobjectDataBlueprintFunctionLibrary.h
@@ -13,6 +13,8 @@
 
 #include "SubobjectDataBlueprintFunctionLibrary.generated.h"
 
+#define UE_API SUBOBJECTDATAINTERFACE_API
+
 class UBlueprint;
 class UObject;
 struct FFrame;
@@ -21,8 +23,8 @@ struct FFrame;
  * A function library with wrappers around the getter/setter functions for FSubobjectData
  * that will make it easier to use within blueprint contexts.
  */
-UCLASS()
-class SUBOBJECTDATAINTERFACE_API USubobjectDataBlueprintFunctionLibrary : public UBlueprintFunctionLibrary
+UCLASS(MinimalAPI)
+class USubobjectDataBlueprintFunctionLibrary : public UBlueprintFunctionLibrary
 {
 	GENERATED_BODY()
 public:
@@ -33,7 +35,7 @@ public:
 	static bool IsHandleValid(const FSubobjectDataHandle& DataHandle) { return DataHandle.IsValid(); }
 
 	UFUNCTION(BlueprintCallable, BlueprintPure, Category="Subobject Data")
-    static void GetData(const FSubobjectDataHandle& DataHandle, FSubobjectData& OutData);
+    static UE_API void GetData(const FSubobjectDataHandle& DataHandle, FSubobjectData& OutData);
 	
 	/**
 	* @return Get the handle for this subobject data
@@ -91,13 +93,13 @@ public:
 	static bool CanRename(const FSubobjectData& Data) { return Data.CanRename(); }
 	
 	UFUNCTION(BlueprintCallable, BlueprintPure, Category="Subobject Data", meta = (DeprecatedFunction="Note", DeprecationMessage = "Use GetAssociatedObject - it will not return template objects for instances"))
-	static const UObject* GetObject(const FSubobjectData& Data, bool bEvenIfPendingKill = false);
+	static UE_API const UObject* GetObject(const FSubobjectData& Data, bool bEvenIfPendingKill = false);
 
 	UFUNCTION(BlueprintCallable, BlueprintPure, Category="Subobject Data")
-	static const UObject* GetAssociatedObject(const FSubobjectData& Data);
+	static UE_API const UObject* GetAssociatedObject(const FSubobjectData& Data);
 
 	UFUNCTION(BlueprintCallable, Category="Subobject Data")
-	static const UObject* GetObjectForBlueprint(const FSubobjectData& Data, UBlueprint* Blueprint);
+	static UE_API const UObject* GetObjectForBlueprint(const FSubobjectData& Data, UBlueprint* Blueprint);
 
 	UFUNCTION(BlueprintCallable, BlueprintPure, Category="Subobject Data")
 	static UBlueprint* GetBlueprint(const FSubobjectData& Data) { return Data.GetBlueprint(); }
@@ -136,3 +138,5 @@ public:
 	UFUNCTION(BlueprintCallable, BlueprintPure, Category="Subobject Data")
 	static bool IsActor(const FSubobjectData& Data) { return Data.IsActor(); }
 };
+
+#undef UE_API
diff --git a/Engine/Source/Editor/SubobjectDataInterface/Public/SubobjectDataHandle.h b/Engine/Source/Editor/SubobjectDataInterface/Public/SubobjectDataHandle.h
--- a/Engine/Source/Editor/SubobjectDataInterface/Public/SubobjectDataHandle.h
+++ b/Engine/Source/Editor/SubobjectDataInterface/Public/SubobjectDataHandle.h
@@ -11,6 +11,8 @@
 
 #include "SubobjectDataHandle.generated.h"
 
+#define UE_API SUBOBJECTDATAINTERFACE_API
+
 struct FSubobjectData;
 
 /**
@@ -20,7 +22,7 @@ struct FSubobjectData;
 * data has validated that it has a good context.
 */
 USTRUCT(BlueprintType)
-struct SUBOBJECTDATAINTERFACE_API FSubobjectDataHandle
+struct FSubobjectDataHandle
 {
 	friend struct FSubobjectData;
 	friend class USubobjectDataSubsystem;
@@ -61,7 +63,7 @@ struct SUBOBJECTDATAINTERFACE_API FSubobjectDataHandle
 	* A static representation of an invalid handle. Mostly used as a return
 	* value for functions that need to return a handle but have the risk of being invalid.
 	*/
-	static const FSubobjectDataHandle InvalidHandle;
+	static UE_API const FSubobjectDataHandle InvalidHandle;
 	
 	/** Get the hash code to use for the given FSubobjectDataHandle */
 	friend uint32 GetTypeHash(const FSubobjectDataHandle& Key)
@@ -82,3 +84,5 @@ struct TStructOpsTypeTraits<FSubobjectDataHandle> : public TStructOpsTypeTraitsB
 		WithIdenticalViaEquality = true,
 	};
 };
+
+#undef UE_API
diff --git a/Engine/Source/Editor/SubobjectDataInterface/Public/SubobjectDataSubsystem.h b/Engine/Source/Editor/SubobjectDataInterface/Public/SubobjectDataSubsystem.h
--- a/Engine/Source/Editor/SubobjectDataInterface/Public/SubobjectDataSubsystem.h
+++ b/Engine/Source/Editor/SubobjectDataInterface/Public/SubobjectDataSubsystem.h
@@ -9,6 +9,8 @@
 #include "ModuleDescriptor.h"
 #include "SubobjectDataSubsystem.generated.h"
 
+#define UE_API SUBOBJECTDATAINTERFACE_API
+
 class UActorComponent;
 class FScopedTransaction;
 
@@ -89,23 +91,23 @@ struct FReparentSubobjectParams
 * instance clicked on via the level editor, a UBlueprint* by opening an asset,
 * or something piped in from python or other scripting languages. 
 */
-UCLASS(Category = "SubobjectDataSubsystem")
-class SUBOBJECTDATAINTERFACE_API USubobjectDataSubsystem : public UEngineSubsystem
+UCLASS(MinimalAPI, Category = "SubobjectDataSubsystem")
+class USubobjectDataSubsystem : public UEngineSubsystem
 {
 	GENERATED_BODY()
 public:
 
 	/** Implement this for initialization of instances of the system */
-	virtual void Initialize(FSubsystemCollectionBase& Collection) override;
+	UE_API virtual void Initialize(FSubsystemCollectionBase& Collection) override;
 
 	/** Implement this for deinitialization of instances of the system */
-	virtual void Deinitialize() override;
+	UE_API virtual void Deinitialize() override;
 	
 	/**
 	* Static wrapper for getting this engine subsystem. Will return nullptr
 	* if the SubobjectDataInterface module has not been loaded yet.
 	*/
-	static USubobjectDataSubsystem* Get();
+	static UE_API USubobjectDataSubsystem* Get();
 	
 	/**
 	* Gather all subobjects that the given UObject context has. Populates an array of 
@@ -114,7 +116,7 @@ public:
 	* @param Context		Object to gather subobjects for
 	* @param OutArray		Array to populate (will be emptied first)
 	*/
-	void GatherSubobjectData(UObject* Context, TArray<FSubobjectDataHandle>& OutArray);
+	UE_API void GatherSubobjectData(UObject* Context, TArray<FSubobjectDataHandle>& OutArray);
 
 	////////////////////////////////////////////
 	// System events
@@ -132,7 +134,7 @@ public:
 	* @param OutArray		Array to populate (will be emptied first)
 	*/
 	UFUNCTION(BlueprintCallable, Category = "SubobjectDataSubsystem", meta = (DisplayName = "Gather Subobject Data For Blueprint"))
-	void K2_GatherSubobjectDataForBlueprint(UBlueprint* Context, TArray<FSubobjectDataHandle>& OutArray);
+	UE_API void K2_GatherSubobjectDataForBlueprint(UBlueprint* Context, TArray<FSubobjectDataHandle>& OutArray);
 
 	/**
 	* Gather all subobjects that the given actor instance has. Populates an array of
@@ -142,7 +144,7 @@ public:
 	* @param OutArray		Array to populate (will be emptied first)
 	*/
 	UFUNCTION(BlueprintCallable, Category = "SubobjectDataSubsystem", meta = (DisplayName = "Gather Subobject Data For Instance"))
-	void K2_GatherSubobjectDataForInstance(AActor* Context, TArray<FSubobjectDataHandle>& OutArray);
+	UE_API void K2_GatherSubobjectDataForInstance(AActor* Context, TArray<FSubobjectDataHandle>& OutArray);
 
 	/**
 	* Attempt to find the subobject data for a given handle. OutData will only 
@@ -154,7 +156,7 @@ public:
 	* @return bool		true if the data was found
 	*/
 	UFUNCTION(BlueprintCallable, Category = "SubobjectDataSubsystem", meta = (DisplayName = "FindSubobjectDataFromHandle"))
-	bool K2_FindSubobjectDataFromHandle(const FSubobjectDataHandle& Handle, FSubobjectData& OutData) const;
+	UE_API bool K2_FindSubobjectDataFromHandle(const FSubobjectDataHandle& Handle, FSubobjectData& OutData) const;
 
 	/**
 	* Attempt to find an existing handle for the given object. 
@@ -165,7 +167,7 @@ public:
 	* @return FSubobjectDataHandle	The subobject handle for the object, Invalid handle if not found.
 	*/
 	UFUNCTION(BlueprintCallable, Category = "SubobjectDataSubsystem")
-	FSubobjectDataHandle FindHandleForObject(const FSubobjectDataHandle& Context, const UObject* ObjectToFind, UBlueprint* BPContext = nullptr) const;
+	UE_API FSubobjectDataHandle FindHandleForObject(const FSubobjectDataHandle& Context, const UObject* ObjectToFind, UBlueprint* BPContext = nullptr) const;
 	
 	////////////////////////////////////////////
 	// Modifying subobjects
@@ -177,7 +179,7 @@ public:
 	 * @return The new class that was created
 	 */
 	UFUNCTION(BlueprintCallable, Category = "SubobjectDataSubsystem", meta = (DisplayName = "Create New C++ Component"))
-	static UClass* CreateNewCPPComponent(TSubclassOf<UActorComponent> ComponentClass, const FString& NewClassPath, const FString& NewClassName);
+	static UE_API UClass* CreateNewCPPComponent(TSubclassOf<UActorComponent> ComponentClass, const FString& NewClassPath, const FString& NewClassName);
 
 	/**
 	 * Creates a new Blueprint component from the specified class type
@@ -186,7 +188,7 @@ public:
 	 * @return The new class that was created
 	 */
 	UFUNCTION(BlueprintCallable, Category = "SubobjectDataSubsystem", meta = (DisplayName = "Create New Blueprint Component"))
-	static UClass* CreateNewBPComponent(TSubclassOf<UActorComponent> ComponentClass, const FString& NewClassPath, const FString& NewClassName);
+	static UE_API UClass* CreateNewBPComponent(TSubclassOf<UActorComponent> ComponentClass, const FString& NewClassPath, const FString& NewClassName);
 	
 	/**
 	* Add a new subobject as a child to the given parent object 
@@ -196,7 +198,7 @@ public:
 	* @return FSubobjectDataHandle		Handle to the newly created subobject, Invalid handle if creation failed
 	*/
 	UFUNCTION(BlueprintCallable, Category = "SubobjectDataSubsystem")
-	FSubobjectDataHandle AddNewSubobject(const FAddNewSubobjectParams& Params, FText& FailReason);
+	UE_API FSubobjectDataHandle AddNewSubobject(const FAddNewSubobjectParams& Params, FText& FailReason);
 
 	/**
 	* Attempts to delete the given array of subobjects from their context 
@@ -208,7 +210,7 @@ public:
 	* @return 	The number of subobjects successfully deleted
 	*/
 	UFUNCTION(BlueprintCallable, Category = "SubobjectDataSubsystem", meta = (DisplayName = "Delete Subobjects from Blueprint"))
-	int32 DeleteSubobjects(const FSubobjectDataHandle& ContextHandle, const TArray<FSubobjectDataHandle>& SubobjectsToDelete, UBlueprint* BPContext = nullptr);
+	UE_API int32 DeleteSubobjects(const FSubobjectDataHandle& ContextHandle, const TArray<FSubobjectDataHandle>& SubobjectsToDelete, UBlueprint* BPContext = nullptr);
 
 	/**
 	* Attempts to delete the given array of subobjects from their context 
@@ -219,7 +221,7 @@ public:
 	* @return 	The number of subobjects successfully deleted
 	*/
 	UFUNCTION(BlueprintCallable, Category = "SubobjectDataSubsystem", meta = (DisplayName = "Delete Subobjects from Instance"))
-	int32 K2_DeleteSubobjectsFromInstance(const FSubobjectDataHandle& ContextHandle, const TArray<FSubobjectDataHandle>& SubobjectsToDelete);
+	UE_API int32 K2_DeleteSubobjectsFromInstance(const FSubobjectDataHandle& ContextHandle, const TArray<FSubobjectDataHandle>& SubobjectsToDelete);
 	
 	/**
 	* Attempts to delete the given array of subobjects from their context 
@@ -230,7 +232,7 @@ public:
 	* @param bForce					If true then this will attempt to delete a subobject even if the CanDelete function flags otherwise.
 	* @return 	The number of subobjects successfully deleted
 	*/
-	int32 DeleteSubobjects(const FSubobjectDataHandle& ContextHandle, const TArray<FSubobjectDataHandle>& SubobjectsToDelete, FSubobjectDataHandle& OutComponentToSelect, UBlueprint* BPContext = nullptr, bool bForce = false);
+	UE_API int32 DeleteSubobjects(const FSubobjectDataHandle& ContextHandle, const TArray<FSubobjectDataHandle>& SubobjectsToDelete, FSubobjectDataHandle& OutComponentToSelect, UBlueprint* BPContext = nullptr, bool bForce = false);
 
 	/**
 	* Attempts to delete the given subobject from its blueprint context 
@@ -242,7 +244,7 @@ public:
 	* @return 	The number of subobjects successfully deleted
 	*/
 	UFUNCTION(BlueprintCallable, Category = "SubobjectDataSubsystem", meta = (DisplayName = "Delete Subobject from Blueprint"))
-	int32 DeleteSubobject(const FSubobjectDataHandle& ContextHandle, const FSubobjectDataHandle& SubobjectToDelete, UBlueprint* BPContext = nullptr);
+	UE_API int32 DeleteSubobject(const FSubobjectDataHandle& ContextHandle, const FSubobjectDataHandle& SubobjectToDelete, UBlueprint* BPContext = nullptr);
 	
 	/**
 	* Attempts to delete the given subobject from its context 
@@ -253,7 +255,7 @@ public:
 	* @return 	The number of subobjects successfully deleted
 	*/
 	UFUNCTION(BlueprintCallable, Category = "SubobjectDataSubsystem", meta = (DisplayName = "Delete Subobject from Instance"))
-	int32 K2_DeleteSubobjectFromInstance(const FSubobjectDataHandle& ContextHandle, const FSubobjectDataHandle& SubobjectToDelete);
+	UE_API int32 K2_DeleteSubobjectFromInstance(const FSubobjectDataHandle& ContextHandle, const FSubobjectDataHandle& SubobjectToDelete);
 	
 	/**
 	* Attempts to rename the given subobject to the new name.
@@ -264,7 +266,7 @@ public:
 	* @return True if the rename was successful, false otherwise.
 	*/
 	UFUNCTION(BlueprintCallable, Category = "SubobjectDataSubsystem")
-	bool RenameSubobject(const FSubobjectDataHandle& Handle, const FText& InNewName);
+	UE_API bool RenameSubobject(const FSubobjectDataHandle& Handle, const FText& InNewName);
 
 	/**
 	* Attempts to change the subclass of a native component
@@ -275,7 +277,7 @@ public:
 	* @return True if the class change was successful, false otherwise.
 	*/
 	UFUNCTION(BlueprintCallable, Category = "SubobjectDataSubsystem")
-	bool ChangeSubobjectClass(const FSubobjectDataHandle& Handle, const UClass* NewClass);
+	UE_API bool ChangeSubobjectClass(const FSubobjectDataHandle& Handle, const UClass* NewClass);
 
 	/**
 	* Attempts to reparent the given subobject to the new parent
@@ -286,10 +288,10 @@ public:
 	* @return True if the reparent was successful, false otherwise.
 	*/
 	UFUNCTION(BlueprintCallable, Category = "SubobjectDataSubsystem")
-    bool ReparentSubobject(const FReparentSubobjectParams& Params, const FSubobjectDataHandle& ToReparentHandle);
+    UE_API bool ReparentSubobject(const FReparentSubobjectParams& Params, const FSubobjectDataHandle& ToReparentHandle);
 
 	UFUNCTION(BlueprintCallable, Category = "SubobjectDataSubsystem")
-	bool MakeNewSceneRoot(const FSubobjectDataHandle& Context, const FSubobjectDataHandle& NewSceneRoot, UBlueprint* BPContext);
+	UE_API bool MakeNewSceneRoot(const FSubobjectDataHandle& Context, const FSubobjectDataHandle& NewSceneRoot, UBlueprint* BPContext);
 	
 	/**
 	* Attempts to reparent all subobjects in the HandlesToMove array to the new parent handle.
@@ -297,7 +299,7 @@ public:
 	* @param NewParentHandle		Handle of the new parent
 	*/
 	UFUNCTION(BlueprintCallable, Category = "SubobjectDataSubsystem")
-	bool ReparentSubobjects(const FReparentSubobjectParams& Params, const TArray<FSubobjectDataHandle>& HandlesToMove);
+	UE_API bool ReparentSubobjects(const FReparentSubobjectParams& Params, const TArray<FSubobjectDataHandle>& HandlesToMove);
 
 	/**
 	 * Remove the child subobject from the owner
@@ -305,7 +307,7 @@ public:
 	 * @return True if the child was successfully removed. 
 	 */
 	UFUNCTION(BlueprintCallable, Category = "SubobjectDataSubsystem")
-	bool DetachSubobject(const FSubobjectDataHandle& OwnerHandle, const FSubobjectDataHandle& ChildToRemove);
+	UE_API bool DetachSubobject(const FSubobjectDataHandle& OwnerHandle, const FSubobjectDataHandle& ChildToRemove);
 
 	/**
 	 * Add the given subobject to a new owner. This will remove the subobject from its previous 
@@ -317,7 +319,7 @@ public:
 	 * @return true if the child was added successfully
 	 */
 	UFUNCTION(BlueprintCallable, Category = "SubobjectDataSubsystem")
-    bool AttachSubobject(const FSubobjectDataHandle& OwnerHandle, const FSubobjectDataHandle& ChildToAddHandle);
+    UE_API bool AttachSubobject(const FSubobjectDataHandle& OwnerHandle, const FSubobjectDataHandle& ChildToAddHandle);
 	
 	/**
 	 * Returns true if the given new text is a valid option to rename the
@@ -331,20 +333,20 @@ public:
 	 * @return True if the rename is valid
 	 */
 	UFUNCTION(BlueprintCallable, Category = "SubobjectDataSubsystem")
-	bool IsValidRename(const FSubobjectDataHandle& Handle, const FText& InNewText, FText& OutErrorMessage) const;
+	UE_API bool IsValidRename(const FSubobjectDataHandle& Handle, const FText& InNewText, FText& OutErrorMessage) const;
 
 	/**
 	 * Returns true if the given array of handles represents subobjects that
 	 * can be copied.
 	 */
 	UFUNCTION(BlueprintCallable, Category = "SubobjectDataSubsystem")
-	bool CanCopySubobjects(const TArray<FSubobjectDataHandle>& Handles) const;
+	UE_API bool CanCopySubobjects(const TArray<FSubobjectDataHandle>& Handles) const;
 
 	UFUNCTION(BlueprintCallable, Category = "SubobjectDataSubsystem")
-	void CopySubobjects(const TArray<FSubobjectDataHandle>& Handles, UBlueprint* BpContext);
+	UE_API void CopySubobjects(const TArray<FSubobjectDataHandle>& Handles, UBlueprint* BpContext);
 
 	UFUNCTION(BlueprintCallable, Category = "SubobjectDataSubsystem")
-	bool CanPasteSubobjects(const FSubobjectDataHandle& RootHandle, UBlueprint* BPContext = nullptr) const;
+	UE_API bool CanPasteSubobjects(const FSubobjectDataHandle& RootHandle, UBlueprint* BPContext = nullptr) const;
 
 	/**
 	* Pastes the given subobjects to the PasteToContext.
@@ -355,7 +357,7 @@ public:
 	* @param OutPastedHandles	Array populated with the handles to the newly pasted subobjects
 	*/
 	UFUNCTION(BlueprintCallable, Category = "SubobjectDataSubsystem")
-	void PasteSubobjects(const FSubobjectDataHandle& PasteToContext, const TArray<FSubobjectDataHandle>& NewParentHandles, UBlueprint* BpContext, TArray<FSubobjectDataHandle>& OutPastedHandles);
+	UE_API void PasteSubobjects(const FSubobjectDataHandle& PasteToContext, const TArray<FSubobjectDataHandle>& NewParentHandles, UBlueprint* BpContext, TArray<FSubobjectDataHandle>& OutPastedHandles);
 	
 	/**
 	 * Duplicate the given array of subobjects on the context. 
@@ -366,19 +368,19 @@ public:
 	 * @param OutNewSubobjects	Array that will be populated with any newly created subobjects
 	 */
 	UFUNCTION(BlueprintCallable, Category = "SubobjectDataSubsystem")
-	void DuplicateSubobjects(const FSubobjectDataHandle& Context, const TArray<FSubobjectDataHandle>& SubobjectsToDup, UBlueprint* BpContext, TArray<FSubobjectDataHandle>& OutNewSubobjects);
+	UE_API void DuplicateSubobjects(const FSubobjectDataHandle& Context, const TArray<FSubobjectDataHandle>& SubobjectsToDup, UBlueprint* BpContext, TArray<FSubobjectDataHandle>& OutNewSubobjects);
 	
-	FScopedTransaction* BeginTransaction(const TArray<FSubobjectDataHandle>& Handles, const FText& Description, UBlueprint* Blueprint);
+	UE_API FScopedTransaction* BeginTransaction(const TArray<FSubobjectDataHandle>& Handles, const FText& Description, UBlueprint* Blueprint);
 
 	/**
 	* Find the scene root for a given subobject handle.
 	* If this is an actor, then it will return the handle pointing to its DefaultSceneRoot.
 	* If this is a component, it will walk the parent chain until it finds the DefaultSceneRoot.
 	*/
-	FSubobjectDataHandle FindSceneRootForSubobject(const FSubobjectDataHandle& InHandle) const;
+	UE_API FSubobjectDataHandle FindSceneRootForSubobject(const FSubobjectDataHandle& InHandle) const;
 
 	UFUNCTION(BlueprintCallable, Category = "SubobjectDataSubsystem")
-	static void RenameSubobjectMemberVariable(UBlueprint* BPContext, const FSubobjectDataHandle& InHandle, const FName NewName);
+	static UE_API void RenameSubobjectMemberVariable(UBlueprint* BPContext, const FSubobjectDataHandle& InHandle, const FName NewName);
 	
 	/**
 	* Walk the hierarchy up to find the AActor root for the given handle
@@ -387,48 +389,50 @@ public:
 	*
 	* @return The handle of the root actor if one exists, FSubobjectDataHandle::InvalidHandle otherwise.
 	*/
-	static FSubobjectDataHandle GetActorRootHandle(const FSubobjectDataHandle& StartingHandle);
+	static UE_API FSubobjectDataHandle GetActorRootHandle(const FSubobjectDataHandle& StartingHandle);
 
 private:
 
 	/**
 	* Recursively add the attachment children of the given actor component
 	*/
-	FSubobjectDataHandle FindOrCreateAttachParentForComponent(UActorComponent* InActorComponent, const FSubobjectDataHandle& ActorRootHandle, TArray<FSubobjectDataHandle>& ExistingHandles);
+	UE_API FSubobjectDataHandle FindOrCreateAttachParentForComponent(UActorComponent* InActorComponent, const FSubobjectDataHandle& ActorRootHandle, TArray<FSubobjectDataHandle>& ExistingHandles);
 
 	/**
 	* Recursively search for an attachment parent that matches the given InActorComponent on the blueprint
 	*/
-	FSubobjectDataHandle FindAttachParentForInheritedComponent(UActorComponent* InActorComponent, const FSubobjectDataHandle& ActorRootHandle, UBlueprint* BP);
+	UE_API FSubobjectDataHandle FindAttachParentForInheritedComponent(UActorComponent* InActorComponent, const FSubobjectDataHandle& ActorRootHandle, UBlueprint* BP);
 
 	/**
 	 * Find all child subobject data related to the given root context.
 	 * Populates the OutVisited set with all data. 
 	 */
-	void FindAllSubobjectData(FSubobjectData* RootContext, TSet<FSubobjectData*>& OutVisited) const;
+	UE_API void FindAllSubobjectData(FSubobjectData* RootContext, TSet<FSubobjectData*>& OutVisited) const;
 	
 	/**
 	 * Gets a ref to subobject data that has been initialized with the given context.
 	 * It will have a unique ID, but may be reused from the FreeData array. 
 	 */
-	FSubobjectDataHandle CreateSubobjectData(UObject* Context, const FSubobjectDataHandle& ParentHandle = FSubobjectDataHandle::InvalidHandle, bool bIsInheritedSCS = false);
+	UE_API FSubobjectDataHandle CreateSubobjectData(UObject* Context, const FSubobjectDataHandle& ParentHandle = FSubobjectDataHandle::InvalidHandle, bool bIsInheritedSCS = false);
 
 	/**
 	 * Get a handle to the subobject with the given context with a specific parent.
 	 * If the context is already a known child of the parent, then it will return an existing subobject handle. 
 	 */
-	FSubobjectDataHandle FactoryCreateSubobjectDataWithParent(UObject* Context, const FSubobjectDataHandle& ParentHandle, bool bIsInheritedSCS = false);
+	UE_API FSubobjectDataHandle FactoryCreateSubobjectDataWithParent(UObject* Context, const FSubobjectDataHandle& ParentHandle, bool bIsInheritedSCS = false);
 
 	/**
 	 * Create a subobject data and handle for an inherited blueprint context.
 	 * This will populate the given array of handles recursively with all possible children
 	 * (which are USCS_Node* in this case)
 	 */
-	FSubobjectDataHandle FactoryCreateInheritedBpSubobject(UObject* Context, const FSubobjectDataHandle& ParentHandle, bool bIsInherited, TArray<FSubobjectDataHandle>& OutArray);
+	UE_API FSubobjectDataHandle FactoryCreateInheritedBpSubobject(UObject* Context, const FSubobjectDataHandle& ParentHandle, bool bIsInherited, TArray<FSubobjectDataHandle>& OutArray);
 	
 	/** Find the parent that this new subobject should attach to */
-	FSubobjectDataHandle FindParentForNewSubobject(const UObject* NewSubobject, const FSubobjectDataHandle& SelectedParent);
+	UE_API FSubobjectDataHandle FindParentForNewSubobject(const UObject* NewSubobject, const FSubobjectDataHandle& SelectedParent);
 
 	FOnNewSubobjectAdded OnNewSubobjectAdded_Delegate;
 	
-};
\ No newline at end of file
+};
+
+#undef UE_API
diff --git a/Engine/Source/Editor/SubobjectEditor/Public/SSubobjectBlueprintEditor.h b/Engine/Source/Editor/SubobjectEditor/Public/SSubobjectBlueprintEditor.h
--- a/Engine/Source/Editor/SubobjectEditor/Public/SSubobjectBlueprintEditor.h
+++ b/Engine/Source/Editor/SubobjectEditor/Public/SSubobjectBlueprintEditor.h
@@ -4,12 +4,14 @@
 
 #include "SSubobjectEditor.h"
 
+#define UE_API SUBOBJECTEDITOR_API
+
 class IClassViewerFilter;
 
 /**
 * This is the editor for subobjects within the blueprint editor that
 */
-class SUBOBJECTEDITOR_API SSubobjectBlueprintEditor final : public SSubobjectEditor
+class SSubobjectBlueprintEditor final : public SSubobjectEditor
 {
 public:
 	DECLARE_DELEGATE_OneParam(FOnHighlightPropertyInDetailsView, const class FPropertyPath&);
@@ -35,7 +37,7 @@ private:
 	    
 	SLATE_END_ARGS()
 	
-	void Construct(const FArguments& InArgs);
+	UE_API void Construct(const FArguments& InArgs);
 
 protected:
 
@@ -44,26 +46,26 @@ protected:
 	
 	// SSubobjectEditor interface
 	virtual bool ShouldModifyBPOnAssetDrop() const override { return true; }
-	virtual void OnDeleteNodes() override;
-	virtual void CopySelectedNodes() override;
-	virtual void OnDuplicateComponent() override;
-	virtual bool CanPasteNodes() const override;
-	virtual void PasteNodes() override;
+	UE_API virtual void OnDeleteNodes() override;
+	UE_API virtual void CopySelectedNodes() override;
+	UE_API virtual void OnDuplicateComponent() override;
+	UE_API virtual bool CanPasteNodes() const override;
+	UE_API virtual void PasteNodes() override;
 	
-	virtual void OnAttachToDropAction(FSubobjectEditorTreeNodePtrType DroppedOn, const TArray<FSubobjectEditorTreeNodePtrType>& DroppedNodePtrs);
-    virtual void OnDetachFromDropAction(const TArray<FSubobjectEditorTreeNodePtrType>& DroppedNodePtrs);
-    virtual void OnMakeNewRootDropAction(FSubobjectEditorTreeNodePtrType DroppedNodePtr);
-    virtual void PostDragDropAction(bool bRegenerateTreeNodes);
+	UE_API virtual void OnAttachToDropAction(FSubobjectEditorTreeNodePtrType DroppedOn, const TArray<FSubobjectEditorTreeNodePtrType>& DroppedNodePtrs);
+    UE_API virtual void OnDetachFromDropAction(const TArray<FSubobjectEditorTreeNodePtrType>& DroppedNodePtrs);
+    UE_API virtual void OnMakeNewRootDropAction(FSubobjectEditorTreeNodePtrType DroppedNodePtr);
+    UE_API virtual void PostDragDropAction(bool bRegenerateTreeNodes);
 
     /** Builds a context menu pop up for dropping a child node onto the scene root node */
-    virtual TSharedPtr<SWidget> BuildSceneRootDropActionMenu(FSubobjectEditorTreeNodePtrType DroppedOntoNodePtr, FSubobjectEditorTreeNodePtrType DroppedNodePtr);
+    UE_API virtual TSharedPtr<SWidget> BuildSceneRootDropActionMenu(FSubobjectEditorTreeNodePtrType DroppedOntoNodePtr, FSubobjectEditorTreeNodePtrType DroppedNodePtr);
 	virtual bool ClearSelectionOnClick() const override { return true; }
 	virtual bool ShowInlineSearchWithButtons() const override { return true; }
-	virtual FSubobjectDataHandle AddNewSubobject(const FSubobjectDataHandle& ParentHandle, UClass* NewClass, UObject* AssetOverride, FText& OutFailReason, TUniquePtr<FScopedTransaction> InOngoingTransaction) override;
-	virtual void PopulateContextMenuImpl(UToolMenu* InMenu, TArray<FSubobjectEditorTreeNodePtrType>& InSelectedItems, bool bIsChildActorSubtreeNodeSelected) override;
-	virtual FSubobjectEditorTreeNodePtrType GetSceneRootNode() const override;
+	UE_API virtual FSubobjectDataHandle AddNewSubobject(const FSubobjectDataHandle& ParentHandle, UClass* NewClass, UObject* AssetOverride, FText& OutFailReason, TUniquePtr<FScopedTransaction> InOngoingTransaction) override;
+	UE_API virtual void PopulateContextMenuImpl(UToolMenu* InMenu, TArray<FSubobjectEditorTreeNodePtrType>& InSelectedItems, bool bIsChildActorSubtreeNodeSelected) override;
+	UE_API virtual FSubobjectEditorTreeNodePtrType GetSceneRootNode() const override;
 public:
-	virtual FSubobjectEditorTreeNodePtrType FindSlateNodeForObject(const UObject* InObject, bool bIncludeAttachmentComponents = true) const override;
+	UE_API virtual FSubobjectEditorTreeNodePtrType FindSlateNodeForObject(const UObject* InObject, bool bIncludeAttachmentComponents = true) const override;
 	// End of SSubobjectEditor
 
 public:
@@ -79,7 +81,7 @@ public:
 	* @param CanExecuteActionDelegate			the delegate to query whether or not to execute the UI action
 	* @param GetSelectedObjectsDelegate		the delegate to fill the currently select variables / components
 	*/
-	static void BuildMenuEventsSection(FMenuBuilder& Menu, UBlueprint* Blueprint, UClass* SelectedClass, FCanExecuteAction CanExecuteActionDelegate, FGetSelectedObjectsDelegate GetSelectedObjectsDelegate);
+	static UE_API void BuildMenuEventsSection(FMenuBuilder& Menu, UBlueprint* Blueprint, UClass* SelectedClass, FCanExecuteAction CanExecuteActionDelegate, FGetSelectedObjectsDelegate GetSelectedObjectsDelegate);
 
 	/**
 	* Highlight a tree node and, optionally, a property with in it
@@ -87,7 +89,7 @@ public:
 	* @param TreeNodeName		Name of the treenode to be highlighted
 	* @param Property	The name of the property to be highlighted in the details view
 	*/
-	void HighlightTreeNode(FName TreeNodeName, const class FPropertyPath& Property);
+	UE_API void HighlightTreeNode(FName TreeNodeName, const class FPropertyPath& Property);
 	
 protected:
 
@@ -98,7 +100,7 @@ protected:
 	* @param GetSelectedObjectsDelegate	the delegate to gather information about current selection
 	* @param NodeIndex						an index to a specified node to add event for or < 0 for all selected nodes.
 	*/
-	static void CreateEventsForSelection(UBlueprint* Blueprint, FName EventName, FGetSelectedObjectsDelegate GetSelectedObjectsDelegate);
+	static UE_API void CreateEventsForSelection(UBlueprint* Blueprint, FName EventName, FGetSelectedObjectsDelegate GetSelectedObjectsDelegate);
 
 	/**
 	* Function to construct an event for a node
@@ -106,7 +108,7 @@ protected:
 	* @param EventName						the event to add
 	* @param EventData						the event data structure describing the node
 	*/
-	static void ConstructEvent(UBlueprint* Blueprint, const FName EventName, const FComponentEventConstructionData EventData);
+	static UE_API void ConstructEvent(UBlueprint* Blueprint, const FName EventName, const FComponentEventConstructionData EventData);
 
 	/**
 	* Function to view an event for a node
@@ -114,8 +116,10 @@ protected:
 	* @param EventName						the event to view
 	* @param EventData						the event data structure describing the node
 	*/
-	static void ViewEvent(UBlueprint* Blueprint, const FName EventName, const FComponentEventConstructionData EventData);
+	static UE_API void ViewEvent(UBlueprint* Blueprint, const FName EventName, const FComponentEventConstructionData EventData);
 
 	/** Get the current preview actor for this blueprint editor. */
 	AActor* GetActorPreview() const { return PreviewActor.Get(nullptr); }
-};
\ No newline at end of file
+};
+
+#undef UE_API
diff --git a/Engine/Source/Editor/SubobjectEditor/Public/SSubobjectEditor.h b/Engine/Source/Editor/SubobjectEditor/Public/SSubobjectEditor.h
--- a/Engine/Source/Editor/SubobjectEditor/Public/SSubobjectEditor.h
+++ b/Engine/Source/Editor/SubobjectEditor/Public/SSubobjectEditor.h
@@ -17,6 +17,8 @@
 
 #include "ScopedTransaction.h"
 
+#define UE_API SUBOBJECTEDITOR_API
+
 class SSubobjectEditor;
 class SToolTip;
 class SExtensionPanel;
@@ -38,14 +40,14 @@ using FSubobjectEditorTreeNodePtrType = TSharedPtr<class FSubobjectEditorTreeNod
  * Wrapper struct that represents access an singe subobject that is used 
  * by slate as a layout for columns/rows should look like
  */
-class SUBOBJECTEDITOR_API FSubobjectEditorTreeNode : public TSharedFromThis<FSubobjectEditorTreeNode>
+class FSubobjectEditorTreeNode : public TSharedFromThis<FSubobjectEditorTreeNode>
 {
 public:
 
 	/** Delegate for when the context menu requests a rename */
 	DECLARE_DELEGATE(FOnRenameRequested);
 	
-	explicit FSubobjectEditorTreeNode(const FSubobjectDataHandle& DataSource, bool InbIsSeperator = false);
+	UE_API explicit FSubobjectEditorTreeNode(const FSubobjectDataHandle& DataSource, bool InbIsSeperator = false);
 
 	~FSubobjectEditorTreeNode() = default;
 
@@ -53,12 +55,12 @@ public:
 
 	// If this subobject is an actor component, then return its template. Otherwise it will
 	// return null
-	const UActorComponent* GetComponentTemplate() const;
+	UE_API const UActorComponent* GetComponentTemplate() const;
 
-	const UObject* GetObject(bool bEvenIfPendingKill = false) const;
+	UE_API const UObject* GetObject(bool bEvenIfPendingKill = false) const;
 
 	// Get a pointer to the subobject data that this slate node is representing
-	FSubobjectData* GetDataSource() const;
+	UE_API FSubobjectData* GetDataSource() const;
 
 	// Get the subobject handle that this slate node is representing
 	FSubobjectDataHandle GetDataHandle() const { return DataHandle; }
@@ -67,35 +69,35 @@ public:
 	FSubobjectEditorTreeNodePtrType GetParent() const { return ParentNodePtr; }
 
 	/** Get the display name of this single tree node */
-	FString GetDisplayString() const;
+	UE_API FString GetDisplayString() const;
 
-	bool IsComponentNode() const;
+	UE_API bool IsComponentNode() const;
 
 	// Returns true if this is the root actor node of the tree
-	bool IsRootActorNode() const;
+	UE_API bool IsRootActorNode() const;
 
-	bool CanReparent() const;
+	UE_API bool CanReparent() const;
 
-	FName GetVariableName() const;
+	UE_API FName GetVariableName() const;
 
 	// Returns true if this subobject data handle is valid
 	bool IsValid() const { return bIsSeperator || DataHandle.IsValid(); }
 
-	bool IsChildSubtreeNode() const;
+	UE_API bool IsChildSubtreeNode() const;
 
-	bool IsNativeComponent() const;
+	UE_API bool IsNativeComponent() const;
 	
 	// Returns true if this node is attached to the given slate node
-	bool IsAttachedTo(FSubobjectEditorTreeNodePtrType InNodePtr) const;
+	UE_API bool IsAttachedTo(FSubobjectEditorTreeNodePtrType InNodePtr) const;
 
 	/**
 	* @return Whether or not this node is a direct child of the given node.
 	*/
-	bool IsDirectlyAttachedTo(FSubobjectEditorTreeNodePtrType InNodePtr) const;
+	UE_API bool IsDirectlyAttachedTo(FSubobjectEditorTreeNodePtrType InNodePtr) const;
 	
-	bool CanDelete() const;
+	UE_API bool CanDelete() const;
 
-	bool CanRename() const;
+	UE_API bool CanRename() const;
 	
 	/**
 	 * @return The set of nodes which are parented to this node (read-only).
@@ -105,34 +107,34 @@ public:
 	// Add the given slate node to our child array and set it's parent to us
 	// This had no effect on the actual structure of the subobjects this node represents
 	// it is purely visual
-	void AddChild(FSubobjectEditorTreeNodePtrType AttachToPtr);
+	UE_API void AddChild(FSubobjectEditorTreeNodePtrType AttachToPtr);
 
 	// Remove the given slate node from our children array
-	void RemoveChild(FSubobjectEditorTreeNodePtrType InChildNodePtr);
+	UE_API void RemoveChild(FSubobjectEditorTreeNodePtrType InChildNodePtr);
 
 	// Attempts to find the given subobject handle in the slate children on this node. Nullptr if none are found.
-	FSubobjectEditorTreeNodePtrType FindChild(const FSubobjectDataHandle& InHandle);
+	UE_API FSubobjectEditorTreeNodePtrType FindChild(const FSubobjectDataHandle& InHandle);
 
 	/** Query that determines if this item should be filtered out or not */
-	bool IsFlaggedForFiltration() const;
+	UE_API bool IsFlaggedForFiltration() const;
 
 	/** Sets this item's filtration state. Use bUpdateParent to make sure the parent's EFilteredState::ChildMatches flag is properly updated based off the new state */
-	void SetCachedFilterState(bool bMatchesFilter, bool bUpdateParent);
+	UE_API void SetCachedFilterState(bool bMatchesFilter, bool bUpdateParent);
 
 	/** Used to update the EFilteredState::ChildMatches flag for parent nodes, when this item's filtration state has changed */
-	void ApplyFilteredStateToParent();
+	UE_API void ApplyFilteredStateToParent();
 
 	/** Updates the EFilteredState::ChildMatches flag, based off of children's current state */
-	void RefreshCachedChildFilterState(bool bUpdateParent);
+	UE_API void RefreshCachedChildFilterState(bool bUpdateParent);
 
 	/** Refreshes this item's filtration state. Set bRecursive to 'true' to refresh any child nodes as well */
-	bool RefreshFilteredState(const UClass* InFilterType, const TArray<FString>& InFilterTerms, bool bRecursive);
+	UE_API bool RefreshFilteredState(const UClass* InFilterType, const TArray<FString>& InFilterTerms, bool bRecursive);
 
 	/** Returns whether the node will match the given type (for filtering) */
-	bool MatchesFilterType(const UClass* InFilterType) const;
+	UE_API bool MatchesFilterType(const UClass* InFilterType) const;
 
-	void SetOngoingCreateTransaction(TUniquePtr<FScopedTransaction> InTransaction);
-	void CloseOngoingCreateTransaction();
+	UE_API void SetOngoingCreateTransaction(TUniquePtr<FScopedTransaction> InTransaction);
+	UE_API void CloseOngoingCreateTransaction();
 	void GetOngoingCreateTransaction(TUniquePtr<FScopedTransaction>& OutPtr) { OutPtr = MoveTemp(OngoingCreateTransaction); }
 	bool HasOngoingTransaction() const { return OngoingCreateTransaction.IsValid(); } 
 
@@ -394,7 +396,7 @@ private:
 * all subobjects of a given actor or blueprint. A subobject can be a native
 * component on an actor or inherited components from a native or blueprint parent.
 */
-class SUBOBJECTEDITOR_API SSubobjectEditor : public SCompoundWidget
+class SSubobjectEditor : public SCompoundWidget
 {
 public:
 
@@ -430,83 +432,83 @@ protected:
 	// Widget Callbacks
 protected:
 
-	TSharedRef<ITableRow> MakeTableRowWidget(FSubobjectEditorTreeNodePtrType InNodePtr, const TSharedRef<STableViewBase>& OwnerTable);
+	UE_API TSharedRef<ITableRow> MakeTableRowWidget(FSubobjectEditorTreeNodePtrType InNodePtr, const TSharedRef<STableViewBase>& OwnerTable);
 
 	/** @return The visibility of the components tree */
-	EVisibility GetComponentsTreeVisibility() const;
+	UE_API EVisibility GetComponentsTreeVisibility() const;
 
 	/** Used by tree control - get children for a specified subobject node */
-	void OnGetChildrenForTree(FSubobjectEditorTreeNodePtrType InNodePtr, TArray<FSubobjectEditorTreeNodePtrType>& OutChildren);
+	UE_API void OnGetChildrenForTree(FSubobjectEditorTreeNodePtrType InNodePtr, TArray<FSubobjectEditorTreeNodePtrType>& OutChildren);
 
 	/** Update any associated selection (e.g. details view) from the passed in nodes */
-	void UpdateSelectionFromNodes(const TArray<FSubobjectEditorTreeNodePtrType>& SelectedNodes);
+	UE_API void UpdateSelectionFromNodes(const TArray<FSubobjectEditorTreeNodePtrType>& SelectedNodes);
 
 	/** Called when selection in the tree changes */
-	void OnTreeSelectionChanged(FSubobjectEditorTreeNodePtrType InSelectedNodePtr, ESelectInfo::Type SelectInfo);
+	UE_API void OnTreeSelectionChanged(FSubobjectEditorTreeNodePtrType InSelectedNodePtr, ESelectInfo::Type SelectInfo);
 
 	/** Callback when a component item is double clicked. */
-	void HandleItemDoubleClicked(FSubobjectEditorTreeNodePtrType InItem);
+	UE_API void HandleItemDoubleClicked(FSubobjectEditorTreeNodePtrType InItem);
 
-	void OnFindReferences(bool bSearchAllBlueprints, const EGetFindReferenceSearchStringFlags Flags);
+	UE_API void OnFindReferences(bool bSearchAllBlueprints, const EGetFindReferenceSearchStringFlags Flags);
 
 	/** @return The visibility of the components filter box */
-	EVisibility GetComponentsFilterBoxVisibility() const;
+	UE_API EVisibility GetComponentsFilterBoxVisibility() const;
 
 	/** Recursively updates the filtered state for each component item */
-	void OnFilterTextChanged(const FText& InFilterText);
+	UE_API void OnFilterTextChanged(const FText& InFilterText);
 
 	/** Callback when a component item is scrolled into view */
-	void OnItemScrolledIntoView(FSubobjectEditorTreeNodePtrType InItem, const TSharedPtr<ITableRow>& InWidget);
+	UE_API void OnItemScrolledIntoView(FSubobjectEditorTreeNodePtrType InItem, const TSharedPtr<ITableRow>& InWidget);
 
 	/** SWidget interface */
-	virtual FReply OnKeyDown(const FGeometry& MyGeometry, const FKeyEvent& InKeyEvent) override;
+	UE_API virtual FReply OnKeyDown(const FGeometry& MyGeometry, const FKeyEvent& InKeyEvent) override;
 	
 public:
 	/** Select the given tree node */
-	void SelectNode(FSubobjectEditorTreeNodePtrType InNodeToSelect, bool bIsCntrlDown);
+	UE_API void SelectNode(FSubobjectEditorTreeNodePtrType InNodeToSelect, bool bIsCntrlDown);
 
 	/** Select the given tree node if there is on that matches the given handle */
-	void SelectNodeFromHandle(const FSubobjectDataHandle& InHandle, bool bIsCntrlDown);
+	UE_API void SelectNodeFromHandle(const FSubobjectDataHandle& InHandle, bool bIsCntrlDown);
 	
 	/** Update the contents of the Subobject Tree based on the current context */
-	void UpdateTree(bool bRegenerateTreeNodes = true);
+	UE_API void UpdateTree(bool bRegenerateTreeNodes = true);
 
 	/** Dumps out the tree view contents to the log (used to assist with debugging widget hierarchy issues) */
-	void DumpTree();
+	UE_API void DumpTree();
 	
 	/** Forces the details panel to refresh on the same objects */
-	void RefreshSelectionDetails();
+	UE_API void RefreshSelectionDetails();
 
 	/** Clears the current selection */
-	void ClearSelection();
+	UE_API void ClearSelection();
 
 	/** Select the root of the tree */
-	void SelectRoot();
+	UE_API void SelectRoot();
 
 	/** Get the tint color that any subobject editor icons should use. The default is the foreground color. */
-	virtual FSlateColor GetColorTintForIcon(FSubobjectEditorTreeNodePtrType Node) const;
+	UE_API virtual FSlateColor GetColorTintForIcon(FSubobjectEditorTreeNodePtrType Node) const;
 	
 	/** Get the currently selected nodes from the tree sorted in order from parent to child */
-	TArray<FSubobjectEditorTreeNodePtrType> GetSelectedNodes() const;
+	UE_API TArray<FSubobjectEditorTreeNodePtrType> GetSelectedNodes() const;
 
 	/** Returns the number of currently selected nodes in the tree */
 	int32 GetNumSelectedNodes() const { return TreeWidget->GetSelectedItems().Num(); }
 
 	/** Get the currently selected handles from the tree sorted in order from parent to child */
-	TArray<FSubobjectDataHandle> GetSelectedHandles() const;
+	UE_API TArray<FSubobjectDataHandle> GetSelectedHandles() const;
 
-	virtual FSubobjectEditorTreeNodePtrType GetSceneRootNode() const;
+	UE_API virtual FSubobjectEditorTreeNodePtrType GetSceneRootNode() const;
 	
 	/** Try to handle a drag-drop operation */
-	FReply TryHandleAssetDragDropOperation(const FDragDropEvent& DragDropEvent);
+	UE_API FReply TryHandleAssetDragDropOperation(const FDragDropEvent& DragDropEvent);
 	
 	/** Sets UI customizations of this SCSEditor. */
-	void SetUICustomization(TSharedPtr<ISCSEditorUICustomization> InUICustomization);
+	UE_API void SetUICustomization(TSharedPtr<ISCSEditorUICustomization> InUICustomization);
 
-	virtual bool IsEditingAllowed() const;
+	UE_API virtual bool IsEditingAllowed() const;
 	
 	/** Return the button widgets that can add components or create/edit blueprints */
-	TSharedPtr<SWidget> GetToolButtonsBox() const;
+	UE_API TSharedPtr<SWidget> GetToolButtonsBox() const;
 	
 	TSharedPtr<FUICommandList> GetCommandList() const { return CommandList; }
 
@@ -532,75 +534,75 @@ public:
 	virtual bool CanMakeNewRootOnDrag(UBlueprint* DraggedFromBlueprint) const { return false; }
 
 	/** Provides access to the Blueprint context that's being edited */
-	UBlueprint* GetBlueprint() const;
+	UE_API UBlueprint* GetBlueprint() const;
 
 	// Attempt to find an existing slate node that matches the given handle
-	virtual FSubobjectEditorTreeNodePtrType FindSlateNodeForObject(const UObject* InObject, bool bIncludeAttachmentComponents = true) const;
+	UE_API virtual FSubobjectEditorTreeNodePtrType FindSlateNodeForObject(const UObject* InObject, bool bIncludeAttachmentComponents = true) const;
 	
 	/** Returns true if the specified component is currently selected */
-    bool IsComponentSelected(const UPrimitiveComponent* PrimComponent) const;
+    UE_API bool IsComponentSelected(const UPrimitiveComponent* PrimComponent) const;
 
     /** Assigns a selection override delegate to the specified component */
-    void SetSelectionOverride(UPrimitiveComponent* PrimComponent) const;
+    UE_API void SetSelectionOverride(UPrimitiveComponent* PrimComponent) const;
 
 	/**
 	* Fills the supplied array with the currently selected objects
 	* @param OutSelectedItems The array to fill.
 	*/
-	void GetSelectedItemsForContextMenu(TArray<FComponentEventConstructionData>& OutSelectedItems) const;
+	UE_API void GetSelectedItemsForContextMenu(TArray<FComponentEventConstructionData>& OutSelectedItems) const;
 
 	/** Return an array of the current slate node hierarchy in the tree */
 	const TArray<FSubobjectEditorTreeNodePtrType>& GetRootNodes() const { return RootNodes; };
 
 	// Attempt to find an existing slate node that matches the given handle
-	FSubobjectEditorTreeNodePtrType FindSlateNodeForHandle(const FSubobjectDataHandle& Handle, FSubobjectEditorTreeNodePtrType InStartNodePtr = FSubobjectEditorTreeNodePtrType()) const;
+	UE_API FSubobjectEditorTreeNodePtrType FindSlateNodeForHandle(const FSubobjectDataHandle& Handle, FSubobjectEditorTreeNodePtrType InStartNodePtr = FSubobjectEditorTreeNodePtrType()) const;
 
 	// Attempt to find an existing slate node that has a given variable name
-	FSubobjectEditorTreeNodePtrType FindSlateNodeForVariableName(FName InVariableName) const;
+	UE_API FSubobjectEditorTreeNodePtrType FindSlateNodeForVariableName(FName InVariableName) const;
 
 	/** Pointer to the current object that is represented by the subobject editor */
-	UObject* GetObjectContext() const;
+	UE_API UObject* GetObjectContext() const;
 
 	/** SubobjectHandle of the current object that is represented by the subobject editor */
-	FSubobjectDataHandle GetObjectContextHandle() const;
+	UE_API FSubobjectDataHandle GetObjectContextHandle() const;
 
 	/** Refresh the type list presented by the add component button when clicked */
-	void RefreshComponentTypesList();
+	UE_API void RefreshComponentTypesList();
 	
 protected:
 
 	/** Restore the previous selection state when updating the tree */
-	virtual void RestoreSelectionState(TArray<FSubobjectEditorTreeNodePtrType>& SelectedTreeNodes, bool bFallBackToVariableName = true);
+	UE_API virtual void RestoreSelectionState(TArray<FSubobjectEditorTreeNodePtrType>& SelectedTreeNodes, bool bFallBackToVariableName = true);
 
 	/** If true, then the blueprint should be modified on TryHandleAssetDragDropOperation */
 	virtual bool ShouldModifyBPOnAssetDrop() const { return false; }
 	
-	bool CanCutNodes() const;
-	void CutSelectedNodes();
+	UE_API bool CanCutNodes() const;
+	UE_API void CutSelectedNodes();
 
-	bool CanCopyNodes() const;
+	UE_API bool CanCopyNodes() const;
 	virtual void CopySelectedNodes() = 0;
 
 	/** Pastes previously copied node(s) */
-	virtual bool CanPasteNodes() const;
+	UE_API virtual bool CanPasteNodes() const;
 	virtual void PasteNodes() = 0;
 	
 public:
 	/** Removes existing selected component nodes from the SCS */
-	virtual bool CanDeleteNodes() const;
+	UE_API virtual bool CanDeleteNodes() const;
 	virtual void OnDeleteNodes() = 0;
 	
 protected:
-	bool CanDuplicateComponent() const;
+	UE_API bool CanDuplicateComponent() const;
 	virtual void OnDuplicateComponent() = 0;
 
 	/** Checks to see if renaming is allowed on the selected component */
-	bool CanRenameComponent() const;
-	void OnRenameComponent();
-	void OnRenameComponent(TUniquePtr<FScopedTransaction> InComponentCreateTransaction);
+	UE_API bool CanRenameComponent() const;
+	UE_API void OnRenameComponent();
+	UE_API void OnRenameComponent(TUniquePtr<FScopedTransaction> InComponentCreateTransaction);
 	
 	/** Called at the end of each frame. */
-	void OnPostTick(float);
+	UE_API void OnPostTick(float);
 	
 	/////////////////////////////////////////////////////////
 	// Widgets
@@ -651,22 +653,22 @@ protected:
 	 */
 	FSubobjectDataHandle CachedRootHandle;
 
-	virtual FMenuBuilder CreateMenuBuilder();
+	UE_API virtual FMenuBuilder CreateMenuBuilder();
 
 	/** Constructs the slate drag/drop tree for this subobject editor */
-	virtual void ConstructTreeWidget();
+	UE_API virtual void ConstructTreeWidget();
 
 	/** Creates a list of commands */
-	virtual void CreateCommandList();
+	UE_API virtual void CreateCommandList();
 
 	/** Recursively visits the given node + its children and invokes the given function for each. */
-	void DepthFirstTraversal(const FSubobjectEditorTreeNodePtrType& InNodePtr, TSet<FSubobjectEditorTreeNodePtrType>& OutVisitedNodes, const TFunctionRef<void(const FSubobjectEditorTreeNodePtrType&)> InFunction) const;
+	UE_API void DepthFirstTraversal(const FSubobjectEditorTreeNodePtrType& InNodePtr, TSet<FSubobjectEditorTreeNodePtrType>& OutVisitedNodes, const TFunctionRef<void(const FSubobjectEditorTreeNodePtrType&)> InFunction) const;
 
 	/** Returns the set of expandable nodes that are currently collapsed in the UI */
-	void GetCollapsedNodes(const FSubobjectEditorTreeNodePtrType& InNodePtr, TSet<FSubobjectEditorTreeNodePtrType>& OutCollapsedNodes) const;
+	UE_API void GetCollapsedNodes(const FSubobjectEditorTreeNodePtrType& InNodePtr, TSet<FSubobjectEditorTreeNodePtrType>& OutCollapsedNodes) const;
 
 	// Attempt to find an existing slate node that matches the given handle
-	static FSubobjectEditorTreeNodePtrType FindOrCreateSlateNodeForHandle(const FSubobjectDataHandle& Handle, TMap<FSubobjectDataHandle, FSubobjectEditorTreeNodePtrType>& ExistingNodes);
+	static UE_API FSubobjectEditorTreeNodePtrType FindOrCreateSlateNodeForHandle(const FSubobjectDataHandle& Handle, TMap<FSubobjectDataHandle, FSubobjectEditorTreeNodePtrType>& ExistingNodes);
 
 	/**
 	 * Set the expansion state of a node
@@ -674,44 +676,44 @@ protected:
 	 * @param InNodeToChange	The node to be expanded/collapsed
 	 * @param bIsExpanded		True to expand the node, false to collapse it
 	 */
-	void SetNodeExpansionState(FSubobjectEditorTreeNodePtrType InNodeToChange, const bool bIsExpanded);
+	UE_API void SetNodeExpansionState(FSubobjectEditorTreeNodePtrType InNodeToChange, const bool bIsExpanded);
 	
 	/** Handler for recursively expanding/collapsing items */
-	void SetItemExpansionRecursive(FSubobjectEditorTreeNodePtrType Model, bool bInExpansionState);
+	UE_API void SetItemExpansionRecursive(FSubobjectEditorTreeNodePtrType Model, bool bInExpansionState);
 
 	/** Registers context menu by name for later access */
-	void RegisterContextMenu();
+	UE_API void RegisterContextMenu();
 
 	/** Called to display context menu when right clicking on the widget */
-	TSharedPtr<SWidget> CreateContextMenu();
+	UE_API TSharedPtr<SWidget> CreateContextMenu();
 
 	/** Populate context menu on the fly */
-	void PopulateContextMenu(UToolMenu* InMenu);
+	UE_API void PopulateContextMenu(UToolMenu* InMenu);
 
 	/** Populate the context menu with implementation specific details */
 	virtual void PopulateContextMenuImpl(UToolMenu* InMenu, TArray<FSubobjectEditorTreeNodePtrType>& InSelectedItems, bool bIsChildActorSubtreeNodeSelected) = 0;
 	
 	/** @return Type of component to filter the tree view with or nullptr if there's no filter. */
-	TSubclassOf<UActorComponent> GetComponentTypeFilterToApply() const;
+	UE_API TSubclassOf<UActorComponent> GetComponentTypeFilterToApply() const;
 
 	/**
 	 * Compares the filter bar's text with the item's component name. Use
 	 * bRecursive to refresh the state of child nodes as well. Returns true if
 	 * the node is set to be filtered out
 	 */
-	bool RefreshFilteredState(FSubobjectEditorTreeNodePtrType TreeNode, bool bRecursive);
+	UE_API bool RefreshFilteredState(FSubobjectEditorTreeNodePtrType TreeNode, bool bRecursive);
 
 	/** Callback for the action trees to get the filter text */
-	FText GetFilterText() const;
+	UE_API FText GetFilterText() const;
 
 	/** Converts the current actor instance to a blueprint */
-	void PromoteToBlueprint() const;
+	UE_API void PromoteToBlueprint() const;
 
 	/** Called when the promote to blueprint button is clicked */
-	FReply OnPromoteToBlueprintClicked();
+	UE_API FReply OnPromoteToBlueprintClicked();
 
 	/** Opens the blueprint editor for the blueprint being viewed by the scseditor */
-	void OnOpenBlueprintEditor(bool bForceCodeEditing) const;
+	UE_API void OnOpenBlueprintEditor(bool bForceCodeEditing) const;
 
 	/**
 	 * Spawns a new SWindow giving the user options for creating a new C++ component class.
@@ -719,7 +721,7 @@ protected:
 	 *
 	 * @return The new class that was created
 	 */
-	UClass* SpawnCreateNewCppComponentWindow(TSubclassOf<UActorComponent> ComponentClass);
+	UE_API UClass* SpawnCreateNewCppComponentWindow(TSubclassOf<UActorComponent> ComponentClass);
 
 	/**
 	 * Spawns a new SWindow giving the user options for creating a new blueprint component class.
@@ -727,19 +729,19 @@ protected:
 	 *
 	 * @return The new class that was created
 	 */
-	UClass* SpawnCreateNewBPComponentWindow(TSubclassOf<UActorComponent> ComponentClass);
+	UE_API UClass* SpawnCreateNewBPComponentWindow(TSubclassOf<UActorComponent> ComponentClass);
 
 	/** Add a component from the selection in the combo box */
-	FSubobjectDataHandle PerformComboAddClass(TSubclassOf<UActorComponent> ComponentClass, EComponentCreateAction::Type ComponentCreateAction, UObject* AssetOverride);
+	UE_API FSubobjectDataHandle PerformComboAddClass(TSubclassOf<UActorComponent> ComponentClass, EComponentCreateAction::Type ComponentCreateAction, UObject* AssetOverride);
 
 	/** @return The visibility of the promote to blueprint button (only visible with an actor instance that is not created from a blueprint)*/
-	virtual EVisibility GetPromoteToBlueprintButtonVisibility() const;
+	UE_API virtual EVisibility GetPromoteToBlueprintButtonVisibility() const;
 
 	/** @return The visibility of the Edit Blueprint button (only visible with an actor instance that is created from a blueprint)*/
-	virtual EVisibility GetEditBlueprintButtonVisibility() const;
+	UE_API virtual EVisibility GetEditBlueprintButtonVisibility() const;
 
 	/** @return The visibility of the Add Component combo button */
-	virtual EVisibility GetComponentClassComboButtonVisibility() const;
+	UE_API virtual EVisibility GetComponentClassComboButtonVisibility() const;
 
 	/** If true, then the tree widget will clear selection on click. */
 	virtual bool ClearSelectionOnClick() const { return false; }
@@ -757,4 +759,6 @@ protected:
 		/** Populate an array of subobject handles based on an array of the slate node ptr */
 		static void PopulateHandlesArray(const TArray<FSubobjectEditorTreeNodePtrType>& SlateNodePtrs, TArray<FSubobjectDataHandle>& OutHandles);
 	};
-};
\ No newline at end of file
+};
+
+#undef UE_API
diff --git a/Engine/Source/Editor/SubobjectEditor/Public/SSubobjectInstanceEditor.h b/Engine/Source/Editor/SubobjectEditor/Public/SSubobjectInstanceEditor.h
--- a/Engine/Source/Editor/SubobjectEditor/Public/SSubobjectInstanceEditor.h
+++ b/Engine/Source/Editor/SubobjectEditor/Public/SSubobjectInstanceEditor.h
@@ -4,11 +4,13 @@
 
 #include "SSubobjectEditor.h"
 
+#define UE_API SUBOBJECTEDITOR_API
+
 /**
 * This is the editor for subobjects within the level editor that
 * works mainly with component and actor instances. 
 */
-class SUBOBJECTEDITOR_API SSubobjectInstanceEditor final : public SSubobjectEditor
+class SSubobjectInstanceEditor final : public SSubobjectEditor
 {
 private:	
 	SLATE_BEGIN_ARGS(SSubobjectInstanceEditor)
@@ -24,41 +26,43 @@ private:
 	
 	SLATE_END_ARGS()
 	
-	void Construct(const FArguments& InArgs);
+	UE_API void Construct(const FArguments& InArgs);
 
 protected:
 	
 	// SSubobjectEditor interface
-	virtual void OnDeleteNodes() override;
-	virtual void CopySelectedNodes() override;
-	virtual void OnDuplicateComponent() override;
-	virtual void PasteNodes() override;
+	UE_API virtual void OnDeleteNodes() override;
+	UE_API virtual void CopySelectedNodes() override;
+	UE_API virtual void OnDuplicateComponent() override;
+	UE_API virtual void PasteNodes() override;
 	
-	virtual void OnAttachToDropAction(FSubobjectEditorTreeNodePtrType DroppedOn, const TArray<FSubobjectEditorTreeNodePtrType>& DroppedNodePtrs) override;
-    virtual void OnDetachFromDropAction(const TArray<FSubobjectEditorTreeNodePtrType>& DroppedNodePtrs) override;
-    virtual void OnMakeNewRootDropAction(FSubobjectEditorTreeNodePtrType DroppedNodePtr) override;
-    virtual void PostDragDropAction(bool bRegenerateTreeNodes) override;
+	UE_API virtual void OnAttachToDropAction(FSubobjectEditorTreeNodePtrType DroppedOn, const TArray<FSubobjectEditorTreeNodePtrType>& DroppedNodePtrs) override;
+    UE_API virtual void OnDetachFromDropAction(const TArray<FSubobjectEditorTreeNodePtrType>& DroppedNodePtrs) override;
+    UE_API virtual void OnMakeNewRootDropAction(FSubobjectEditorTreeNodePtrType DroppedNodePtr) override;
+    UE_API virtual void PostDragDropAction(bool bRegenerateTreeNodes) override;
 
     /** Builds a context menu popup for dropping a child node onto the scene root node */
-    virtual TSharedPtr<SWidget> BuildSceneRootDropActionMenu(FSubobjectEditorTreeNodePtrType DroppedOntoNodePtr, FSubobjectEditorTreeNodePtrType DroppedNodePtr) override;
-	virtual FSubobjectDataHandle AddNewSubobject(const FSubobjectDataHandle& ParentHandle, UClass* NewClass, UObject* AssetOverride, FText& OutFailReason, TUniquePtr<FScopedTransaction> InOngoingTransaction) override;
-	virtual void PopulateContextMenuImpl(UToolMenu* InMenu, TArray<FSubobjectEditorTreeNodePtrType>& InSelectedItems, bool bIsChildActorSubtreeNodeSelected) override;
-	virtual FMenuBuilder CreateMenuBuilder();
+    UE_API virtual TSharedPtr<SWidget> BuildSceneRootDropActionMenu(FSubobjectEditorTreeNodePtrType DroppedOntoNodePtr, FSubobjectEditorTreeNodePtrType DroppedNodePtr) override;
+	UE_API virtual FSubobjectDataHandle AddNewSubobject(const FSubobjectDataHandle& ParentHandle, UClass* NewClass, UObject* AssetOverride, FText& OutFailReason, TUniquePtr<FScopedTransaction> InOngoingTransaction) override;
+	UE_API virtual void PopulateContextMenuImpl(UToolMenu* InMenu, TArray<FSubobjectEditorTreeNodePtrType>& InSelectedItems, bool bIsChildActorSubtreeNodeSelected) override;
+	UE_API virtual FMenuBuilder CreateMenuBuilder();
 	// End of SSubobjectEditor
 
 public:
 private:
 
 	/** @return the tooltip describing how many properties will be applied to the blueprint */
-	FText OnGetApplyChangesToBlueprintTooltip() const;
+	UE_API FText OnGetApplyChangesToBlueprintTooltip() const;
 
 	/** Propagates instance changes to the blueprint */
-	void OnApplyChangesToBlueprint() const;
+	UE_API void OnApplyChangesToBlueprint() const;
 
 	/** Resets instance changes to the blueprint default */
-	void OnResetToBlueprintDefaults();
+	UE_API void OnResetToBlueprintDefaults();
 
 	/** @return the tooltip describing how many properties will be reset to the blueprint default*/
-	FText OnGetResetToBlueprintDefaultsTooltip() const;
+	UE_API FText OnGetResetToBlueprintDefaultsTooltip() const;
+
+};
 
-};
\ No newline at end of file
+#undef UE_API
diff --git a/Engine/Source/Editor/SubobjectEditor/Public/SubobjectEditorMenuContext.h b/Engine/Source/Editor/SubobjectEditor/Public/SubobjectEditorMenuContext.h
--- a/Engine/Source/Editor/SubobjectEditor/Public/SubobjectEditorMenuContext.h
+++ b/Engine/Source/Editor/SubobjectEditor/Public/SubobjectEditorMenuContext.h
@@ -11,19 +11,23 @@
 
 #include "SubobjectEditorMenuContext.generated.h"
 
+#define UE_API SUBOBJECTEDITOR_API
+
 class SSubobjectEditor;
 struct FFrame;
 
-UCLASS()
-class SUBOBJECTEDITOR_API USubobjectEditorMenuContext : public UObject
+UCLASS(MinimalAPI)
+class USubobjectEditorMenuContext : public UObject
 {
 	GENERATED_BODY()
 public:
 
 	UFUNCTION(BlueprintCallable, Category="Tool Menus")
-	TArray<UObject*> GetSelectedObjects() const;
+	UE_API TArray<UObject*> GetSelectedObjects() const;
 
 	TWeakPtr<SSubobjectEditor> SubobjectEditor;
 	
 	bool bOnlyShowPasteOption;
 };
+
+#undef UE_API
diff --git a/Engine/Source/Editor/TextureEditor/Classes/TextureEditorSettings.h b/Engine/Source/Editor/TextureEditor/Classes/TextureEditorSettings.h
--- a/Engine/Source/Editor/TextureEditor/Classes/TextureEditorSettings.h
+++ b/Engine/Source/Editor/TextureEditor/Classes/TextureEditorSettings.h
@@ -52,8 +52,8 @@ enum class ETextureEditorZoomMode : uint8
 /**
  * Implements the Editor's user settings.
  */
-UCLASS(config=EditorPerProjectUserSettings)
-class TEXTUREEDITOR_API UTextureEditorSettings
+UCLASS(MinimalAPI, config=EditorPerProjectUserSettings)
+class UTextureEditorSettings
 	: public UObject
 {
 	GENERATED_UCLASS_BODY()
diff --git a/Engine/Source/Editor/UMGEditor/Classes/WidgetPaletteFavorites.h b/Engine/Source/Editor/UMGEditor/Classes/WidgetPaletteFavorites.h
--- a/Engine/Source/Editor/UMGEditor/Classes/WidgetPaletteFavorites.h
+++ b/Engine/Source/Editor/UMGEditor/Classes/WidgetPaletteFavorites.h
@@ -10,14 +10,16 @@
 
 #include "WidgetPaletteFavorites.generated.h"
 
-UCLASS(config = EditorPerProjectUserSettings)
-class UMGEDITOR_API UWidgetPaletteFavorites : public UObject
+#define UE_API UMGEDITOR_API
+
+UCLASS(MinimalAPI, config = EditorPerProjectUserSettings)
+class UWidgetPaletteFavorites : public UObject
 {
 	GENERATED_UCLASS_BODY()
 public:
 
-	void Add(const FString& InWidgetTemplateName);
-	void Remove(const FString& InWidgetTemplateName);
+	UE_API void Add(const FString& InWidgetTemplateName);
+	UE_API void Remove(const FString& InWidgetTemplateName);
 
 	TArray<FString> GetFavorites() const { return Favorites; }
 	
@@ -30,3 +32,5 @@ private:
 	UPROPERTY(config)
 	TArray<FString> Favorites;
 };
+
+#undef UE_API
diff --git a/Engine/Source/Editor/UMGEditor/Private/Nodes/K2Node_CreateDragDropOperation.h b/Engine/Source/Editor/UMGEditor/Private/Nodes/K2Node_CreateDragDropOperation.h
--- a/Engine/Source/Editor/UMGEditor/Private/Nodes/K2Node_CreateDragDropOperation.h
+++ b/Engine/Source/Editor/UMGEditor/Private/Nodes/K2Node_CreateDragDropOperation.h
@@ -7,28 +7,32 @@
 #include "K2Node_ConstructObjectFromClass.h"
 #include "K2Node_CreateDragDropOperation.generated.h"
 
+#define UE_API UMGEDITOR_API
+
 class UEdGraph;
 
-UCLASS()
-class UMGEDITOR_API UK2Node_CreateDragDropOperation : public UK2Node_ConstructObjectFromClass
+UCLASS(MinimalAPI)
+class UK2Node_CreateDragDropOperation : public UK2Node_ConstructObjectFromClass
 {
 	GENERATED_UCLASS_BODY()
 
 	//~ Begin UEdGraphNode Interface.
-	virtual void AllocateDefaultPins() override;
-	virtual void ExpandNode(class FKismetCompilerContext& CompilerContext, UEdGraph* SourceGraph) override;
+	UE_API virtual void AllocateDefaultPins() override;
+	UE_API virtual void ExpandNode(class FKismetCompilerContext& CompilerContext, UEdGraph* SourceGraph) override;
 	//~ End UEdGraphNode Interface.
 
 	//~ Begin UK2Node Interface
-	virtual FText GetMenuCategory() const override;
-	virtual FName GetCornerIcon() const override;
+	UE_API virtual FText GetMenuCategory() const override;
+	UE_API virtual FName GetCornerIcon() const override;
 	//~ End UK2Node Interface.
 
 protected:
 	/** Gets the default node title when no class is selected */
-	virtual FText GetBaseNodeTitle() const override;
+	UE_API virtual FText GetBaseNodeTitle() const override;
 	/** Gets the node title when a class has been selected. */
-	virtual FText GetNodeTitleFormat() const override;
+	UE_API virtual FText GetNodeTitleFormat() const override;
 	/** Gets base class to use for the 'class' pin.  UObject by default. */
-	virtual UClass* GetClassPinBaseClass() const override;
+	UE_API virtual UClass* GetClassPinBaseClass() const override;
 };
+
+#undef UE_API
diff --git a/Engine/Source/Editor/UMGEditor/Private/Nodes/K2Node_CreateWidget.h b/Engine/Source/Editor/UMGEditor/Private/Nodes/K2Node_CreateWidget.h
--- a/Engine/Source/Editor/UMGEditor/Private/Nodes/K2Node_CreateWidget.h
+++ b/Engine/Source/Editor/UMGEditor/Private/Nodes/K2Node_CreateWidget.h
@@ -7,34 +7,38 @@
 #include "K2Node_ConstructObjectFromClass.h"
 #include "K2Node_CreateWidget.generated.h"
 
+#define UE_API UMGEDITOR_API
+
 class UEdGraph;
 
-UCLASS()
-class UMGEDITOR_API UK2Node_CreateWidget : public UK2Node_ConstructObjectFromClass
+UCLASS(MinimalAPI)
+class UK2Node_CreateWidget : public UK2Node_ConstructObjectFromClass
 {
 	GENERATED_UCLASS_BODY()
 
 	//~ Begin UEdGraphNode Interface.
-	virtual void AllocateDefaultPins() override;
-	virtual void ExpandNode(class FKismetCompilerContext& CompilerContext, UEdGraph* SourceGraph) override;
-	virtual TSharedPtr<SGraphNode> CreateVisualWidget() override;
+	UE_API virtual void AllocateDefaultPins() override;
+	UE_API virtual void ExpandNode(class FKismetCompilerContext& CompilerContext, UEdGraph* SourceGraph) override;
+	UE_API virtual TSharedPtr<SGraphNode> CreateVisualWidget() override;
 	//~ End UEdGraphNode Interface.
 
 	//~ Begin UK2Node Interface
-	virtual FText GetMenuCategory() const override;
-	virtual FName GetCornerIcon() const override;
+	UE_API virtual FText GetMenuCategory() const override;
+	UE_API virtual FName GetCornerIcon() const override;
 	//~ End UK2Node Interface.
 
 	/** Get the owning player pin */
-	UEdGraphPin* GetOwningPlayerPin() const;
+	UE_API UEdGraphPin* GetOwningPlayerPin() const;
 
 protected:
 	/** Gets the default node title when no class is selected */
-	virtual FText GetBaseNodeTitle() const override;
+	UE_API virtual FText GetBaseNodeTitle() const override;
 	/** Gets the node title when a class has been selected. */
-	virtual FText GetNodeTitleFormat() const override;
+	UE_API virtual FText GetNodeTitleFormat() const override;
 	/** Gets base class to use for the 'class' pin.  UObject by default. */
-	virtual UClass* GetClassPinBaseClass() const override;
+	UE_API virtual UClass* GetClassPinBaseClass() const override;
 	/**  */
-	virtual bool IsSpawnVarPin(UEdGraphPin* Pin) const override;
+	UE_API virtual bool IsSpawnVarPin(UEdGraphPin* Pin) const override;
 };
+
+#undef UE_API
diff --git a/Engine/Source/Editor/UMGEditor/Private/Nodes/K2Node_PlayAnimation.h b/Engine/Source/Editor/UMGEditor/Private/Nodes/K2Node_PlayAnimation.h
--- a/Engine/Source/Editor/UMGEditor/Private/Nodes/K2Node_PlayAnimation.h
+++ b/Engine/Source/Editor/UMGEditor/Private/Nodes/K2Node_PlayAnimation.h
@@ -5,26 +5,26 @@
 #include "K2Node_BaseAsyncTask.h"
 #include "K2Node_PlayAnimation.generated.h"
 
-UCLASS()
-class UMGEDITOR_API UK2Node_PlayAnimation : public UK2Node_BaseAsyncTask
+UCLASS(MinimalAPI)
+class UK2Node_PlayAnimation : public UK2Node_BaseAsyncTask
 {
 	GENERATED_UCLASS_BODY()
 };
 
-UCLASS()
-class UMGEDITOR_API UK2Node_PlayAnimationTimeRange : public UK2Node_BaseAsyncTask
+UCLASS(MinimalAPI)
+class UK2Node_PlayAnimationTimeRange : public UK2Node_BaseAsyncTask
 {
 	GENERATED_UCLASS_BODY()
 };
 
-UCLASS()
-class UMGEDITOR_API UK2Node_PlayAnimation2 : public UK2Node_BaseAsyncTask
+UCLASS(MinimalAPI)
+class UK2Node_PlayAnimation2 : public UK2Node_BaseAsyncTask
 {
 	GENERATED_UCLASS_BODY()
 };
 
-UCLASS()
-class UMGEDITOR_API UK2Node_PlayAnimationTimeRange2 : public UK2Node_BaseAsyncTask
+UCLASS(MinimalAPI)
+class UK2Node_PlayAnimationTimeRange2 : public UK2Node_BaseAsyncTask
 {
 	GENERATED_UCLASS_BODY()
 };
