Blueprint Editor: Duplicating multiple elements now selects all duplicated elements.

[RN]

diff --git a/Engine/Source/Editor/SubobjectEditor/Private/SSubobjectBlueprintEditor.cpp b/Engine/Source/Editor/SubobjectEditor/Private/SSubobjectBlueprintEditor.cpp
--- a/Engine/Source/Editor/SubobjectEditor/Private/SSubobjectBlueprintEditor.cpp
+++ b/Engine/Source/Editor/SubobjectEditor/Private/SSubobjectBlueprintEditor.cpp
@@ -287,31 +287,62 @@ void SSubobjectBlueprintEditor::CopySelectedNodes()
 
 void SSubobjectBlueprintEditor::OnDuplicateComponent()
 {
-	TArray<FSubobjectDataHandle> SelectedNodes = GetSelectedHandles();
-	if(SelectedNodes.Num() > 0)
+	const TArray<FSubobjectDataHandle> SelectedNodes = GetSelectedHandles();
+	if (SelectedNodes.Num() <= 0)
 	{
-		// Force the text box being edited (if any) to commit its text. The duplicate operation may trigger a regeneration of the tree view,
-		// releasing all row widgets. If one row was in edit mode (rename/rename on create), it was released before losing the focus and
-		// this would prevent the completion of the 'rename' or 'create + give initial name' transaction (occurring on focus lost).
-		FSlateApplication::Get().ClearKeyboardFocus();
-		
-		TUniquePtr<FScopedTransaction> Transaction = MakeUnique<FScopedTransaction>(SelectedNodes.Num() > 1 ? LOCTEXT("DuplicateComponents", "Duplicate Components") : LOCTEXT("DuplicateComponent", "Duplicate Component"));
+		return;
+	}
+
+	// Force the text box being edited (if any) to commit its text. The duplicate operation may trigger a regeneration of the tree view,
+	// releasing all row widgets. If one row was in edit mode (rename/rename on create), it was released before losing the focus and
+	// this would prevent the completion of the 'rename' or 'create + give initial name' transaction (occurring on focus lost).
+	FSlateApplication::Get().ClearKeyboardFocus();
 		
-		if(USubobjectDataSubsystem* System = USubobjectDataSubsystem::Get())
-		{
-			TArray<FSubobjectDataHandle> DuplicatedHandles;
-			System->DuplicateSubobjects(GetObjectContextHandle(), SelectedNodes, GetBlueprint(), DuplicatedHandles);
-			UpdateTree();
+	TUniquePtr<FScopedTransaction> Transaction = MakeUnique<FScopedTransaction>(SelectedNodes.Num() > 1 ? LOCTEXT("DuplicateComponents", "Duplicate Components") : LOCTEXT("DuplicateComponent", "Duplicate Component"));
+	USubobjectDataSubsystem* System = USubobjectDataSubsystem::Get();
+	if (!System)
+	{
+		return;
+	}
+
+	TArray<FSubobjectDataHandle> DuplicatedHandles;
+	System->DuplicateSubobjects(GetObjectContextHandle(), SelectedNodes, GetBlueprint(), DuplicatedHandles);
+	UpdateTree();
+
+	// If nothing was duplicated, no further action is needed
+	if (DuplicatedHandles.Num() <= 0)
+	{
+		return;
+	}
 
-			// Set focus to the newly created subobject
-			FSubobjectEditorTreeNodePtrType NewNode = DuplicatedHandles.Num() > 0 ? FindSlateNodeForHandle(DuplicatedHandles[0]) : nullptr;
-			if (NewNode != nullptr)
+	ClearSelection();
+
+	FSubobjectEditorTreeNodePtrType FirstNode = nullptr;
+
+	// Select all duplicated subobjects
+	for (const FSubobjectDataHandle& DataHandle : DuplicatedHandles)
+	{
+		if (const FSubobjectEditorTreeNodePtrType NewNode = FindSlateNodeForHandle(DataHandle))
+		{
+			if (!FirstNode)
 			{
-				TreeWidget->SetSelection(NewNode);
-				OnRenameComponent(MoveTemp(Transaction));
+				FirstNode = NewNode;
 			}
+
+			SelectNode(NewNode, true);
 		}
 	}
+
+	// If only one subobject was duplicated, start a rename operation on it
+	if (DuplicatedHandles.Num() == 1)
+	{
+		OnRenameComponent(MoveTemp(Transaction));
+	}
+	// Otherwise, scroll the first of the duplicated nodes into view in the tree widget
+	else if (FirstNode)
+	{
+		TreeWidget->RequestScrollIntoView(FirstNode);
+	}
 }
 
 bool SSubobjectBlueprintEditor::CanPasteNodes() const
