Blueprint Editor:
- Selection changes are now undoable.
- Clicking empty space now clears the existing selection.

[RN]

diff --git a/Engine/Source/Editor/Kismet/Private/SCSEditorViewportClient.cpp b/Engine/Source/Editor/Kismet/Private/SCSEditorViewportClient.cpp
--- a/Engine/Source/Editor/Kismet/Private/SCSEditorViewportClient.cpp
+++ b/Engine/Source/Editor/Kismet/Private/SCSEditorViewportClient.cpp
@@ -390,107 +390,119 @@ void FSCSEditorViewportClient::ProcessClick(class FSceneView& View, class HHitPr
 {
 	const FViewportClick Click(&View, this, Key, Event, HitX, HitY);
 
-	if (HitProxy)
+	if (!HitProxy)
 	{
-		if (HitProxy->IsA(HInstancedStaticMeshInstance::StaticGetType()))
+		if (BlueprintEditorPtr.IsValid())
 		{
-			HInstancedStaticMeshInstance* InstancedStaticMeshInstanceProxy = ((HInstancedStaticMeshInstance*)HitProxy);
+			// If the user clicked on empty space, clear the selection
+			const FScopedTransaction Transaction(NSLOCTEXT("UnrealEd", "SelectNone", "Select None"));
+			BlueprintEditorPtr.Pin()->GetSubobjectEditor()->ClearSelection();
+		}
+		return;
+	}
 
-			TSharedPtr<ISCSEditorCustomization> Customization = BlueprintEditorPtr.Pin()->CustomizeSubobjectEditor(InstancedStaticMeshInstanceProxy->Component);
-			if (Customization.IsValid() && Customization->HandleViewportClick(AsShared(), View, HitProxy, Key, Event, HitX, HitY))
-			{
-				Invalidate();
-			}
+	const FScopedTransaction Transaction(NSLOCTEXT("UnrealEd", "ClickingOnElements", "Clicking on Elements"));
+	if (HitProxy->IsA(HInstancedStaticMeshInstance::StaticGetType()))
+	{
+		HInstancedStaticMeshInstance* InstancedStaticMeshInstanceProxy = ((HInstancedStaticMeshInstance*)HitProxy);
 
-			return;
-		}
-		else if (HitProxy->IsA(HWidgetAxis::StaticGetType()))
+		TSharedPtr<ISCSEditorCustomization> Customization = BlueprintEditorPtr.Pin()->CustomizeSubobjectEditor(InstancedStaticMeshInstanceProxy->Component);
+		if (Customization.IsValid() && Customization->HandleViewportClick(AsShared(), View, HitProxy, Key, Event, HitX, HitY))
 		{
-			const bool bOldModeWidgets1 = EngineShowFlags.ModeWidgets;
-			const bool bOldModeWidgets2 = View.Family->EngineShowFlags.ModeWidgets;
-
-			EngineShowFlags.SetModeWidgets(false);
-			FSceneViewFamily* SceneViewFamily = const_cast<FSceneViewFamily*>(View.Family);
-			SceneViewFamily->EngineShowFlags.SetModeWidgets(false);
-			bool bWasWidgetDragging = Widget->IsDragging();
-			Widget->SetDragging(false);
-
-			// Invalidate the hit proxy map so it will be rendered out again when GetHitProxy
-			// is called
-			Viewport->InvalidateHitProxy();
-
-			// This will actually re-render the viewport's hit proxies!
-			HHitProxy* HitProxyWithoutAxisWidgets = Viewport->GetHitProxy(HitX, HitY);
-			if (HitProxyWithoutAxisWidgets != NULL && !HitProxyWithoutAxisWidgets->IsA(HWidgetAxis::StaticGetType()))
-			{
-				// Try this again, but without the widget this time!
-				ProcessClick(View, HitProxyWithoutAxisWidgets, Key, Event, HitX, HitY);
-			}
-
-			// Undo the evil
-			EngineShowFlags.SetModeWidgets(bOldModeWidgets1);
-			SceneViewFamily->EngineShowFlags.SetModeWidgets(bOldModeWidgets2);
+			Invalidate();
+		}
 
-			Widget->SetDragging(bWasWidgetDragging);
+		return;
+	}
 
-			// Invalidate the hit proxy map again so that it'll be refreshed with the original
-			// scene contents if we need it again later.
-			Viewport->InvalidateHitProxy();
-			return;
-		}
-		else if(GUnrealEd->ComponentVisManager.HandleClick(this, HitProxy, Click))
+	if (HitProxy->IsA(HWidgetAxis::StaticGetType()))
+	{
+		const bool bOldModeWidgets1 = EngineShowFlags.ModeWidgets;
+		const bool bOldModeWidgets2 = View.Family->EngineShowFlags.ModeWidgets;
+
+		EngineShowFlags.SetModeWidgets(false);
+		FSceneViewFamily* SceneViewFamily = const_cast<FSceneViewFamily*>(View.Family);
+		SceneViewFamily->EngineShowFlags.SetModeWidgets(false);
+		bool bWasWidgetDragging = Widget->IsDragging();
+		Widget->SetDragging(false);
+
+		// Invalidate the hit proxy map so it will be rendered out again when GetHitProxy
+		// is called
+		Viewport->InvalidateHitProxy();
+
+		// This will actually re-render the viewport's hit proxies!
+		HHitProxy* HitProxyWithoutAxisWidgets = Viewport->GetHitProxy(HitX, HitY);
+		if (HitProxyWithoutAxisWidgets != NULL && !HitProxyWithoutAxisWidgets->IsA(HWidgetAxis::StaticGetType()))
 		{
-			// Component Vis Manager handled this click, no need to do anything
+			// Try this again, but without the widget this time!
+			ProcessClick(View, HitProxyWithoutAxisWidgets, Key, Event, HitX, HitY);
 		}
-		else if (HitProxy->IsA(HActor::StaticGetType()))
+
+		// Undo the evil
+		EngineShowFlags.SetModeWidgets(bOldModeWidgets1);
+		SceneViewFamily->EngineShowFlags.SetModeWidgets(bOldModeWidgets2);
+
+		Widget->SetDragging(bWasWidgetDragging);
+
+		// Invalidate the hit proxy map again so that it'll be refreshed with the original
+		// scene contents if we need it again later.
+		Viewport->InvalidateHitProxy();
+		return;
+	}
+
+	if(GUnrealEd->ComponentVisManager.HandleClick(this, HitProxy, Click))
+	{
+		// Component Vis Manager handled this click, no need to do anything
+		return;
+	}
+
+	if (HitProxy->IsA(HActor::StaticGetType()))
+	{
+		HActor* ActorProxy = (HActor*)HitProxy;
+		AActor* PreviewActor = GetPreviewActor();
+		if (ActorProxy && ActorProxy->Actor && ActorProxy->PrimComponent)
 		{
-			HActor* ActorProxy = (HActor*)HitProxy;
-			AActor* PreviewActor = GetPreviewActor();
-			if (ActorProxy && ActorProxy->Actor && ActorProxy->PrimComponent)
-			{
-				const USceneComponent* SelectedCompInstance = nullptr;
+			const USceneComponent* SelectedCompInstance = nullptr;
 
-				if (ActorProxy->Actor == PreviewActor)
+			if (ActorProxy->Actor == PreviewActor)
+			{
+				const UPrimitiveComponent* TestComponent = ActorProxy->PrimComponent;
+				if (ActorProxy->Actor->GetComponents().Contains(TestComponent)
+					&& IsObjectSelectableInViewport(const_cast<UPrimitiveComponent*>(TestComponent)))
 				{
-					const UPrimitiveComponent* TestComponent = ActorProxy->PrimComponent;
-					if (ActorProxy->Actor->GetComponents().Contains(TestComponent)
-						&& IsObjectSelectableInViewport(const_cast<UPrimitiveComponent*>(TestComponent)))
-					{
-						SelectedCompInstance = TestComponent;
-					}
+					SelectedCompInstance = TestComponent;
 				}
-				else if (ActorProxy->Actor->IsChildActor())
+			}
+			else if (ActorProxy->Actor->IsChildActor())
+			{
+				AActor* TestActor = ActorProxy->Actor;
+				while (TestActor->GetParentActor()->IsChildActor())
 				{
-					AActor* TestActor = ActorProxy->Actor;
-					while (TestActor->GetParentActor()->IsChildActor())
-					{
-						TestActor = TestActor->GetParentActor();
-					}
+					TestActor = TestActor->GetParentActor();
+				}
 
-					if (TestActor->GetParentActor() == PreviewActor)
-					{
-						SelectedCompInstance = TestActor->GetParentComponent();
-					}
+				if (TestActor->GetParentActor() == PreviewActor)
+				{
+					SelectedCompInstance = TestActor->GetParentComponent();
 				}
+			}
 
-				if (SelectedCompInstance)
+			if (SelectedCompInstance)
+			{
+				TSharedPtr<ISCSEditorCustomization> Customization = BlueprintEditorPtr.Pin()->CustomizeSubobjectEditor(SelectedCompInstance);
+				if (!(Customization.IsValid() && Customization->HandleViewportClick(AsShared(), View, HitProxy, Key, Event, HitX, HitY)))
 				{
-					TSharedPtr<ISCSEditorCustomization> Customization = BlueprintEditorPtr.Pin()->CustomizeSubobjectEditor(SelectedCompInstance);
-					if (!(Customization.IsValid() && Customization->HandleViewportClick(AsShared(), View, HitProxy, Key, Event, HitX, HitY)))
+					const bool bIsCtrlKeyDown = Viewport->KeyState(EKeys::LeftControl) || Viewport->KeyState(EKeys::RightControl);
+					if (BlueprintEditorPtr.IsValid())
 					{
-						const bool bIsCtrlKeyDown = Viewport->KeyState(EKeys::LeftControl) || Viewport->KeyState(EKeys::RightControl);
-						if (BlueprintEditorPtr.IsValid())
-						{
-							// Note: This will find and select any node associated with the component instance that's attached to the proxy (including visualizers)
-							BlueprintEditorPtr.Pin()->FindAndSelectSubobjectEditorTreeNode(SelectedCompInstance, bIsCtrlKeyDown);
-						}
+						// Note: This will find and select any node associated with the component instance that's attached to the proxy (including visualizers)
+						BlueprintEditorPtr.Pin()->FindAndSelectSubobjectEditorTreeNode(SelectedCompInstance, bIsCtrlKeyDown);
 					}
 				}
 			}
-
-			Invalidate();
-			return;
 		}
+
+		Invalidate();
 	}
 }
 
diff --git a/Engine/Source/Editor/SubobjectEditor/Private/SSubobjectBlueprintEditor.cpp b/Engine/Source/Editor/SubobjectEditor/Private/SSubobjectBlueprintEditor.cpp
--- a/Engine/Source/Editor/SubobjectEditor/Private/SSubobjectBlueprintEditor.cpp
+++ b/Engine/Source/Editor/SubobjectEditor/Private/SSubobjectBlueprintEditor.cpp
@@ -47,9 +47,10 @@ void SSubobjectBlueprintEditor::Construct(const FArguments& InArgs)
 	OnHighlightPropertyInDetailsView = InArgs._OnHighlightPropertyInDetailsView;
 	AllowEditing = InArgs._AllowEditing;
 	HideComponentClassCombo = InArgs._HideComponentClassCombo;
-	bAllowTreeUpdates = true;
+	bAllowUpdateTree = true;
 	
 	CreateCommandList();
+	CreateSelectionSet();
 
 	// Build the tree widget
 	FSlateBrush const* MobilityHeaderBrush = FAppStyle::GetBrush(TEXT("ClassIcon.ComponentMobilityHeaderIcon"));
@@ -246,14 +247,7 @@ void SSubobjectBlueprintEditor::OnDeleteNodes()
 
 		if(NumDeleted > 0)
 		{
-			if(HandleToSelect.IsValid())
-			{
-				FSubobjectEditorTreeNodePtrType NodeToSelect = FindSlateNodeForHandle(HandleToSelect);
-				if(NodeToSelect.IsValid())
-				{
-					TreeWidget->SetSelection(NodeToSelect);
-				}
-			}
+			SelectNodeFromHandle(HandleToSelect, false);
 			
 			UpdateTree();
 
@@ -270,8 +264,6 @@ void SSubobjectBlueprintEditor::OnDeleteNodes()
 					ThumbnailTools::CacheEmptyThumbnail( BPFullName, BPPackage );
 				}
 			}
-			// Do this AFTER marking the Blueprint as modified
-			UpdateSelectionFromNodes(TreeWidget->GetSelectedItems());
 		}
 	}
 }
@@ -293,6 +285,9 @@ void SSubobjectBlueprintEditor::OnDuplicateComponent()
 		return;
 	}
 
+	// Don't update the selection set based on these updates to the tree widget
+	TGuardValue BlockSelection{ bBlockSelectionUpdates, true };
+
 	// Force the text box being edited (if any) to commit its text. The duplicate operation may trigger a regeneration of the tree view,
 	// releasing all row widgets. If one row was in edit mode (rename/rename on create), it was released before losing the focus and
 	// this would prevent the completion of the 'rename' or 'create + give initial name' transaction (occurring on focus lost).
@@ -315,8 +310,6 @@ void SSubobjectBlueprintEditor::OnDuplicateComponent()
 		return;
 	}
 
-	ClearSelection();
-
 	FSubobjectEditorTreeNodePtrType FirstNode = nullptr;
 
 	// Select all duplicated subobjects
@@ -327,19 +320,30 @@ void SSubobjectBlueprintEditor::OnDuplicateComponent()
 			if (!FirstNode)
 			{
 				FirstNode = NewNode;
-			}
 
-			SelectNode(NewNode, true);
+				// Clear the selection while selecting the first node
+				SelectNode(NewNode, false);
+			}
+			else
+			{
+				SelectNode(NewNode, true);
+			}
 		}
 	}
 
+	// If we weren't able to select any of the new nodes, don't move any further
+	if (!FirstNode)
+	{
+		return;
+	}
+
 	// If only one subobject was duplicated, start a rename operation on it
 	if (DuplicatedHandles.Num() == 1)
 	{
 		OnRenameComponent(MoveTemp(Transaction));
 	}
 	// Otherwise, scroll the first of the duplicated nodes into view in the tree widget
-	else if (FirstNode)
+	else
 	{
 		TreeWidget->RequestScrollIntoView(FirstNode);
 	}
@@ -370,8 +374,8 @@ void SSubobjectBlueprintEditor::PasteNodes()
 		TArray<FSubobjectDataHandle> OutHandles;
 		
 		// stop allowing tree updates during paste, otherwise MarkBlueprintAsModified will trigger in the middle of it
-		bool bRestoreAllowTreeUpdates = bAllowTreeUpdates;
-		bAllowTreeUpdates = false;
+		bool bRestoreAllowTreeUpdates = bAllowUpdateTree;
+		bAllowUpdateTree = false;
 		TArray<FSubobjectDataHandle> HandlesToPasteOnto = GetSelectedHandles();
 		if(HandlesToPasteOnto.IsEmpty())
 		{
@@ -387,20 +391,20 @@ void SSubobjectBlueprintEditor::PasteNodes()
 		System->PasteSubobjects(GetObjectContextHandle(), HandlesToPasteOnto, GetBlueprint(), OutHandles);
 
 		// allow tree updates again
-		bAllowTreeUpdates = bRestoreAllowTreeUpdates;
+		bAllowUpdateTree = bRestoreAllowTreeUpdates;
 		
 		if(OutHandles.Num() > 0)
 		{
 			// We only want the pasted node(s) to be selected
-			TreeWidget->ClearSelection();
+			ClearSelection();
 			UpdateTree();
 
 			for(const FSubobjectDataHandle& Handle : OutHandles)
 			{
+				SelectNodeFromHandle(Handle, true);
 				if(FSubobjectEditorTreeNodePtrType SlateNode = FindSlateNodeForHandle(Handle))
 				{
 					TreeWidget->RequestScrollIntoView(SlateNode);
-					TreeWidget->SetItemSelection(SlateNode, true);	
 				}
 			}
 		}
diff --git a/Engine/Source/Editor/SubobjectEditor/Private/SSubobjectEditor.cpp b/Engine/Source/Editor/SubobjectEditor/Private/SSubobjectEditor.cpp
--- a/Engine/Source/Editor/SubobjectEditor/Private/SSubobjectEditor.cpp
+++ b/Engine/Source/Editor/SubobjectEditor/Private/SSubobjectEditor.cpp
@@ -60,6 +60,7 @@
 #include "K2Node_Variable.h"		// Nodes for FSubobjectRowDragDropOp
 #include "K2Node_VariableGet.h"
 #include "AssetSelection.h"
+#include "Elements/Framework/TypedElementRegistry.h"
 #include "Settings/BlueprintEditorProjectSettings.h"
 #include "Misc/FeedbackContext.h"
 
@@ -2062,6 +2063,52 @@ bool SSubobject_RowWidget::IsReadOnly() const
 ////////////////////////////////////////////////
 // SSubobjectEditor
 
+SSubobjectEditor::~SSubobjectEditor()
+{
+	if (SelectedElements && UObjectInitialized())
+	{
+		if (!SelectedElements->HasAnyFlags(RF_BeginDestroyed | RF_FinishDestroyed))
+		{
+			SelectedElements->ClearSelection(FTypedElementSelectionOptions());
+		}
+		SelectedElements->OnPreChange().Clear();
+		SelectedElements->OnChanged().Clear();
+		SelectedElements->RemoveFromRoot();
+		SelectedElements = nullptr;
+	}
+}
+
+void SSubobjectEditor::OnSelectionChanged(const UTypedElementSelectionSet* InChangedSelectionSet)
+{
+	// If this is in response to a selection change from the tree widget, no action is needed
+	if (bBlockTreeWidgetUpdates && TreeWidget.IsValid())
+	{
+		return;
+	}
+
+	// Don't update the selection set based on these updates to the tree widget
+	TGuardValue BlockSelection{ bBlockSelectionUpdates, true };
+
+	// Clear the existing tree view selection
+	TreeWidget->ClearSelection();
+
+	// Match the tree view to the selection set
+	const UTypedElementSelectionSet& SelectionSet = *InChangedSelectionSet;
+	for (const UObject* SelectedObject : SelectionSet.GetSelectedObjects())
+	{
+		FSubobjectEditorTreeNodePtrType TreeNodePtr = FindSlateNodeForObject(SelectedObject);
+		if (!TreeWidget.IsValid() || !TreeNodePtr.IsValid())
+		{
+			return;
+		}
+
+		TreeWidget->SetItemSelection(TreeNodePtr, true);
+	}
+
+	// Notify that the selection has updated
+	OnSelectionUpdated.ExecuteIfBound(TreeWidget->GetSelectedItems());
+}
+
 FMenuBuilder SSubobjectEditor::CreateMenuBuilder()
 {
 	// Menu builder for editing a blueprint
@@ -2175,6 +2222,13 @@ void SSubobjectEditor::CreateCommandList()
 	);
 }
 
+void SSubobjectEditor::CreateSelectionSet()
+{
+	SelectedElements = NewObject<UTypedElementSelectionSet>(GetTransientPackage(), NAME_None, RF_Transactional);
+	SelectedElements->AddToRoot();
+	SelectedElements->OnChanged().AddSP(this, &SSubobjectEditor::OnSelectionChanged);
+}
+
 bool SSubobjectEditor::IsComponentSelected(const UPrimitiveComponent* PrimComponent) const
 {
 	check(PrimComponent);
@@ -2286,6 +2340,11 @@ TArray<FSubobjectEditorTreeNodePtrType> SSubobjectEditor::GetSelectedNodes() con
 	return SelectedTreeNodes;
 }
 
+UTypedElementSelectionSet* SSubobjectEditor::GetSelectionSet() const
+{
+	return SelectedElements.Get();
+}
+
 TArray<FSubobjectDataHandle> SSubobjectEditor::GetSelectedHandles() const
 {
 	TArray<FSubobjectEditorTreeNodePtrType> SelectedNodes = GetSelectedNodes();
@@ -2316,16 +2375,45 @@ FSubobjectEditorTreeNodePtrType SSubobjectEditor::GetSceneRootNode() const
 
 void SSubobjectEditor::UpdateSelectionFromNodes(const TArray<FSubobjectEditorTreeNodePtrType>& SelectedNodes)
 {
-	bUpdatingSelection = true;
+	// Don't bother if this is ultimately originating from changes made locally to the selection
+	if (bBlockSelectionUpdates)
+	{
+		return;
+	}
 
-	// Notify that the selection has updated
-	OnSelectionUpdated.ExecuteIfBound(SelectedNodes);
+	check(SelectedElements);
+
+	const FTypedElementSelectionOptions SelectionOptions = FTypedElementSelectionOptions()
+		.SetAllowHidden(true)
+		.SetWarnIfLocked(true);
 
-	bUpdatingSelection = false;
+	// Update the SelectedElements to match the new selection in the TreeWidget
+	SelectedElements->ClearSelection(SelectionOptions);
+	for (const FSubobjectEditorTreeNodePtrType& SelectedNode : SelectedNodes)
+	{
+		const FSubobjectData* const Data = SelectedNode->GetDataSource();
+		if (!Data || !Data->GetObject())
+		{
+			continue;
+		}
+
+		FTypedElementHandle TypedElementHandle = UEngineElementsLibrary::AcquireEditorObjectElementHandle(Data->GetObject());
+		SelectedElements->SelectElement(TypedElementHandle, SelectionOptions);
+	}
 }
 
 void SSubobjectEditor::OnTreeSelectionChanged(FSubobjectEditorTreeNodePtrType InSelectedNodePtr, ESelectInfo::Type SelectInfo)
 {
+	// Don't bother if this is ultimately originating from changes made in the tree
+	if (bBlockTreeWidgetUpdates)
+	{
+		return;
+	}
+
+	// Don't update the tree based on updates to the selection
+	TGuardValue BlockTreeWidgetUpdates{ bBlockTreeWidgetUpdates, true };
+
+	const FScopedTransaction Transaction(NSLOCTEXT("UnrealEd", "SelectingElements", "Selecting Elements"), !bBlockSelectionUpdates && SelectInfo != ESelectInfo::Direct);
 	UpdateSelectionFromNodes(TreeWidget->GetSelectedItems());
 }
 
@@ -2362,7 +2450,7 @@ void SSubobjectEditor::SelectRoot()
 {
 	if (RootNodes.Num() > 0)
 	{
-		TreeWidget->SetSelection(RootNodes[0]);
+		SelectNode(RootNodes[0], true);
 	}
 }
 
@@ -2371,24 +2459,63 @@ FSlateColor SSubobjectEditor::GetColorTintForIcon(FSubobjectEditorTreeNodePtrTyp
 	return FSlateColor::UseForeground();
 }
 
+void SSubobjectEditor::SelectSubobject(const FSubobjectData* InSubobjectData, bool bIsCntrlDown)
+{
+	if (!InSubobjectData || !InSubobjectData->GetObject())
+	{
+		return;
+	}
+
+	const FTypedElementHandle TypedElementHandle = UEngineElementsLibrary::AcquireEditorObjectElementHandle(InSubobjectData->GetObject());
+	SelectElement(TypedElementHandle, bIsCntrlDown);
+}
+
 void SSubobjectEditor::SelectNodeFromHandle(const FSubobjectDataHandle& InHandle, bool bIsCntrlDown)
 {
-	SelectNode(FindSlateNodeForHandle(InHandle), bIsCntrlDown);
+	if (!InHandle.IsValid())
+	{
+		return;
+	}
+
+	SelectSubobject(InHandle.GetData(), bIsCntrlDown);
 }
 
 void SSubobjectEditor::SelectNode(FSubobjectEditorTreeNodePtrType InNodeToSelect, bool bIsCntrlDown)
 {
-	if (TreeWidget.IsValid() && InNodeToSelect.IsValid())
+	if (!InNodeToSelect)
+	{
+		return;
+	}
+
+	SelectSubobject(InNodeToSelect->GetDataSource(), bIsCntrlDown);
+}
+
+void SSubobjectEditor::SelectElement(const FTypedElementHandle& InTypedElementHandle, bool bIsCntrlDown)
+{
+	check(SelectedElements);
+
+	const FTypedElementSelectionOptions SelectionOptions = FTypedElementSelectionOptions()
+		.SetAllowHidden(true)
+		.SetWarnIfLocked(true);
+
+	// If Control is held down, toggle the selection of the individual element
+	if (bIsCntrlDown)
 	{
-		if (!bIsCntrlDown)
+		if (SelectedElements->IsElementSelected(InTypedElementHandle, FTypedElementIsSelectedOptions().SetAllowIndirect(true)))
 		{
-			TreeWidget->SetSelection(InNodeToSelect);
+			SelectedElements->DeselectElement(InTypedElementHandle, SelectionOptions);
 		}
 		else
 		{
-			TreeWidget->SetItemSelection(InNodeToSelect, !TreeWidget->IsItemSelected(InNodeToSelect));
+			SelectedElements->SelectElement(InTypedElementHandle, SelectionOptions);
 		}
 	}
+	// Otherwise, select only that individual element
+	else
+	{
+		SelectedElements->ClearSelection(SelectionOptions);
+		SelectedElements->SelectElement(InTypedElementHandle, SelectionOptions);
+	}
 }
 
 EVisibility SSubobjectEditor::GetComponentsFilterBoxVisibility() const
@@ -2483,7 +2610,7 @@ void SSubobjectEditor::UpdateTree(bool bRegenerateTreeNodes /* = true */)
 	check(TreeWidget.IsValid());
 
 	// Early exit if we're deferring tree updates
-	if (!bAllowTreeUpdates)
+	if (!bAllowUpdateTree)
 	{
 		return;
 	}
@@ -2493,6 +2620,9 @@ void SSubobjectEditor::UpdateTree(bool bRegenerateTreeNodes /* = true */)
 		// Obtain the list of selected items
 		TArray<FSubobjectEditorTreeNodePtrType> SelectedTreeNodes = GetSelectedNodes();
 
+		// Don't update the selection set based on these updates to the tree widget
+		TGuardValue BlockSelection{ bBlockSelectionUpdates, true };
+
 		// Clear the current tree
 		if (SelectedTreeNodes.Num() != 0)
 		{
@@ -2601,58 +2731,52 @@ FSubobjectEditorTreeNodePtrType SSubobjectEditor::FindOrCreateSlateNodeForHandle
 
 void SSubobjectEditor::RefreshSelectionDetails()
 {
-	UpdateSelectionFromNodes(TreeWidget->GetSelectedItems());
+	OnSelectionChanged(SelectedElements.Get());
 }
 
 void SSubobjectEditor::ClearSelection()
 {
-	if (bUpdatingSelection == false)
-	{
-		check(TreeWidget.IsValid());
-		TreeWidget->ClearSelection();
-	}
+	check(SelectedElements);
+
+	const FTypedElementSelectionOptions SelectionOptions = FTypedElementSelectionOptions()
+		.SetAllowHidden(true)
+		.SetWarnIfLocked(true);
+	SelectedElements->ClearSelection(SelectionOptions);
 }
 
 void SSubobjectEditor::RestoreSelectionState(TArray<FSubobjectEditorTreeNodePtrType>& SelectedTreeNodes, bool bFallBackToVariableName)
 {
-	if (SelectedTreeNodes.Num() > 0)
+	if (SelectedTreeNodes.Num() <= 0)
 	{
-		// If there is only one item selected, imitate user selection to preserve navigation
-		ESelectInfo::Type SelectInfo = SelectedTreeNodes.Num() == 1 ? ESelectInfo::OnMouseClick : ESelectInfo::Direct;
+		return;
+	}
 
-		// Restore the previous selection state on the new tree nodes
-		for (int i = 0; i < SelectedTreeNodes.Num(); ++i)
+	// If there is only one item selected, imitate user selection to preserve navigation
+	ESelectInfo::Type SelectInfo = SelectedTreeNodes.Num() == 1 ? ESelectInfo::OnMouseClick : ESelectInfo::Direct;
+
+	// Restore the previous selection state on the new tree nodes
+	for (const FSubobjectEditorTreeNodePtrType& SelectedObject : SelectedTreeNodes)
+	{
+		if (RootNodes.Num() > 0 && SelectedObject == RootNodes[0])
+		{
+			TreeWidget->SetItemSelection(RootNodes[0], true, SelectInfo);
+		}
+		else
 		{
-			if (RootNodes.Num() > 0 && SelectedTreeNodes[i] == RootNodes[0])
+			FSubobjectDataHandle CurrentNodeDataHandle = SelectedObject->GetDataHandle();
+			FSubobjectEditorTreeNodePtrType NodeToSelectPtr = FindSlateNodeForHandle(CurrentNodeDataHandle);
+
+			// If we didn't find something for this exact handle, fall back to just search for something
+			// with the same variable name. This helps to still preserve selection across re-compiles of a class.
+			if (!NodeToSelectPtr.IsValid() && CurrentNodeDataHandle.IsValid() && bFallBackToVariableName)
 			{
-				TreeWidget->SetItemSelection(RootNodes[0], true, SelectInfo);
+				NodeToSelectPtr = FindSlateNodeForVariableName(CurrentNodeDataHandle.GetSharedDataPtr()->GetVariableName());
 			}
-			else
-			{
-				FSubobjectDataHandle CurrentNodeDataHandle = SelectedTreeNodes[i]->GetDataHandle();
-				FSubobjectEditorTreeNodePtrType NodeToSelectPtr = FindSlateNodeForHandle(CurrentNodeDataHandle);
-
-				// If we didn't find something for this exact handle, fall back to just search for something
-				// with the same variable name. This helps to still preserve selection across re-compiles of a class.
-				if (!NodeToSelectPtr.IsValid() && CurrentNodeDataHandle.IsValid() && bFallBackToVariableName)
-				{
-					NodeToSelectPtr = FindSlateNodeForVariableName(CurrentNodeDataHandle.GetSharedDataPtr()->GetVariableName());
-				}
 
-				if (NodeToSelectPtr.IsValid())
-				{
-					TreeWidget->SetItemSelection(NodeToSelectPtr, true, SelectInfo);
-				}
+			if (NodeToSelectPtr.IsValid())
+			{
+				TreeWidget->SetItemSelection(NodeToSelectPtr, true, SelectInfo);
 			}
-
-			//if (GetEditorMode() != EComponentEditorMode::BlueprintSCS)
-			//{
-			//	TArray<FSubobjectEditorTreeNodePtrType> NewSelectedTreeNodes = TreeWidget->GetSelectedItems();
-			//	if (NewSelectedTreeNodes.Num() == 0)
-			//	{
-			//		TreeWidget->SetItemSelection(RootNodes[0], true);
-			//	}
-			//}
 		}
 	}
 }
@@ -2774,7 +2898,7 @@ FReply SSubobjectEditor::TryHandleAssetDragDropOperation(const FDragDropEvent& D
 					Blueprint->SimpleConstructionScript->SaveToTransactionBuffer();
 				}
 
-				bAllowTreeUpdates = true;
+				bAllowUpdateTree = true;
 				FBlueprintEditorUtils::MarkBlueprintAsStructurallyModified(Blueprint);
 			}
 
@@ -3573,12 +3697,8 @@ FSubobjectDataHandle SSubobjectEditor::PerformComboAddClass(TSubclassOf<UActorCo
 			UpdateTree();
 			
 			// Set focus to the newly created subobject
-			FSubobjectEditorTreeNodePtrType NewNode = FindSlateNodeForHandle(NewComponentHandle);
-			if(NewNode != nullptr)
-			{
-				TreeWidget->SetSelection(NewNode);
-				OnRenameComponent(MoveTemp(AddTransaction));
-			}
+			SelectSubobject(NewComponentHandle.GetData(), false);
+			OnRenameComponent(MoveTemp(AddTransaction));
 		}
 	}
 
diff --git a/Engine/Source/Editor/SubobjectEditor/Private/SSubobjectInstanceEditor.cpp b/Engine/Source/Editor/SubobjectEditor/Private/SSubobjectInstanceEditor.cpp
--- a/Engine/Source/Editor/SubobjectEditor/Private/SSubobjectInstanceEditor.cpp
+++ b/Engine/Source/Editor/SubobjectEditor/Private/SSubobjectInstanceEditor.cpp
@@ -41,9 +41,10 @@ void SSubobjectInstanceEditor::Construct(const FArguments& InArgs)
 	OnSelectionUpdated = InArgs._OnSelectionUpdated;
 	OnItemDoubleClicked = InArgs._OnItemDoubleClicked;
 	AllowEditing = InArgs._AllowEditing;
-	bAllowTreeUpdates = true;
+	bAllowUpdateTree = true;
 	
 	CreateCommandList();
+	CreateSelectionSet();
 
 	// Build the tree widget
 	FSlateBrush const* MobilityHeaderBrush = FAppStyle::GetBrush(TEXT("ClassIcon.ComponentMobilityHeaderIcon"));
@@ -224,23 +225,20 @@ void SSubobjectInstanceEditor::OnDeleteNodes()
 		int32 NumDeleted = System->DeleteSubobjects(RootNodes[0]->GetDataHandle(), HandlesToDelete, HandleToSelect);
 
 		if(NumDeleted > 0)
-		{			
-			FSubobjectEditorTreeNodePtrType NodeToSelect = HandleToSelect.IsValid() ? FindSlateNodeForHandle(HandleToSelect) : GetSceneRootNode();
-			if(NodeToSelect.IsValid())
+		{
+			// Attempt to select the next handle
+			if (HandleToSelect.IsValid() && HandleToSelect.GetData()->GetObject())
 			{
-				TreeWidget->SetSelection(NodeToSelect);
+				SelectNodeFromHandle(HandleToSelect, false);
 			}
 			// If there are no components left, then fall back to the generic root node.
 			// This may be the case if you have deleted all components on a native instance actor
-			else if(!RootNodes.IsEmpty())
+			else if (!RootNodes.IsEmpty())
 			{
-				TreeWidget->SetSelection(RootNodes[0]);
+				SelectNode(RootNodes[0], false);
 			}
 			
 			UpdateTree();
-
-			// Do this AFTER marking the Blueprint as modified
-			UpdateSelectionFromNodes(TreeWidget->GetSelectedItems());
 		}
 	}
 }
@@ -268,17 +266,39 @@ void SSubobjectInstanceEditor::OnDuplicateComponent()
 
 		if(USubobjectDataSubsystem* System = USubobjectDataSubsystem::Get())
 		{
+			ClearSelection();
+
 			TArray<FSubobjectDataHandle> DuplicatedHandles;
 			System->DuplicateSubobjects(GetObjectContextHandle(), SelectedNodes, /* BpContext = */ nullptr, DuplicatedHandles);
+
 			UpdateTree();
 
-			// Set focus to the newly created subobject
-			FSubobjectEditorTreeNodePtrType NewNode = DuplicatedHandles.Num() > 0 ? FindSlateNodeForHandle(DuplicatedHandles[0]) : nullptr;
-			if (NewNode != nullptr)
+			FSubobjectDataHandle FirstHandle{};
+
+			// Select all duplicated handles
+			for (const FSubobjectDataHandle& DuplicatedHandle : DuplicatedHandles)
+			{
+				if (!FirstHandle.IsValid())
+				{
+					FirstHandle = DuplicatedHandle;
+				}
+
+				SelectNodeFromHandle(DuplicatedHandle, true);
+			}
+
+			// If only one subobject was duplicated, start a rename operation on it
+			if (DuplicatedHandles.Num() == 1)
 			{
-				TreeWidget->SetSelection(NewNode);
 				OnRenameComponent(MoveTemp(Transaction));
 			}
+			// Otherwise, scroll the first of the duplicated nodes into view in the tree widget
+			else if (FirstHandle.IsValid())
+			{
+				if (const FSubobjectEditorTreeNodePtrType NewNode = FindSlateNodeForHandle(FirstHandle))
+				{
+					TreeWidget->RequestScrollIntoView(NewNode);
+				}
+			}
 		}
 	}
 }
@@ -305,15 +325,15 @@ void SSubobjectInstanceEditor::PasteNodes()
 		if(OutHandles.Num() > 0)
 		{
 			// We only want the pasted node(s) to be selected
-			TreeWidget->ClearSelection();
+			ClearSelection();
 			UpdateTree();
 
 			for(const FSubobjectDataHandle& Handle : OutHandles)
 			{
+				SelectNodeFromHandle(Handle, true);
 				if(FSubobjectEditorTreeNodePtrType SlateNode = FindSlateNodeForHandle(Handle))
 				{
 					TreeWidget->RequestScrollIntoView(SlateNode);
-					TreeWidget->SetItemSelection(SlateNode, true);	
 				}
 			}
 		}
diff --git a/Engine/Source/Editor/SubobjectEditor/Public/SSubobjectEditor.h b/Engine/Source/Editor/SubobjectEditor/Public/SSubobjectEditor.h
--- a/Engine/Source/Editor/SubobjectEditor/Public/SSubobjectEditor.h
+++ b/Engine/Source/Editor/SubobjectEditor/Public/SSubobjectEditor.h
@@ -16,6 +16,7 @@
 #include "BlueprintEditor.h"		// FComponentEventConstructionData
 
 #include "ScopedTransaction.h"
+#include "Elements/Interfaces/TypedElementSelectionInterface.h"
 
 #define UE_API SUBOBJECTEDITOR_API
 
@@ -404,11 +405,10 @@ public:
 	DECLARE_DELEGATE_OneParam(FOnItemDoubleClicked, const FSubobjectEditorTreeNodePtrType);
 	
 protected:
-	
 	// Do not allow public construction of this widget! 
 	SSubobjectEditor() = default;
 
-	virtual ~SSubobjectEditor() = default;
+	UE_API virtual ~SSubobjectEditor() override;
 	
 	/** Delegate to invoke on selection update. */
 	FOnSelectionUpdated OnSelectionUpdated;
@@ -464,11 +464,17 @@ protected:
 	UE_API virtual FReply OnKeyDown(const FGeometry& MyGeometry, const FKeyEvent& InKeyEvent) override;
 	
 public:
-	/** Select the given tree node */
+	/** Select the subobject referenced by the given tree node, if it's valid. */
 	UE_API void SelectNode(FSubobjectEditorTreeNodePtrType InNodeToSelect, bool bIsCntrlDown);
 
-	/** Select the given tree node if there is on that matches the given handle */
+	/** Select the given subobject, if it's valid. */
 	UE_API void SelectNodeFromHandle(const FSubobjectDataHandle& InHandle, bool bIsCntrlDown);
+
+	/** Select the given subobject, if it's valid. */
+	UE_API void SelectSubobject(const FSubobjectData* InHandle, bool bIsCntrlDown);
+
+	/** Select the given typed element handle. */
+	UE_API void SelectElement(const FTypedElementHandle& InTypedElementHandle, bool bIsCntrlDown);
 	
 	/** Update the contents of the Subobject Tree based on the current context */
 	UE_API void UpdateTree(bool bRegenerateTreeNodes = true);
@@ -491,6 +497,9 @@ public:
 	/** Get the currently selected nodes from the tree sorted in order from parent to child */
 	UE_API TArray<FSubobjectEditorTreeNodePtrType> GetSelectedNodes() const;
 
+	/** Get the selection set */
+	UE_API UTypedElementSelectionSet* GetSelectionSet() const;
+
 	/** Returns the number of currently selected nodes in the tree */
 	int32 GetNumSelectedNodes() const { return TreeWidget->GetSelectedItems().Num(); }
 
@@ -625,11 +634,17 @@ protected:
 	/** The add component button / type selector */
 	TSharedPtr<SComponentClassCombo> ComponentClassCombo;
 
-	/** Gate to prevent changing the selection while selection change is being broadcast. */
-	bool bUpdatingSelection;
+	/** The list of selected elements. */
+	TObjectPtr<UTypedElementSelectionSet> SelectedElements = nullptr;
+
+	/** Gate to prevent changing the tree widget's selection while responding to its own changes. */
+	bool bBlockTreeWidgetUpdates = false;
+
+	/** Gate to prevent changing the local selection while responding to its own changes. */
+	bool bBlockSelectionUpdates = false;
 
 	/** Controls whether or not to allow calls to UpdateTree() */
-	bool bAllowTreeUpdates;
+	bool bAllowUpdateTree;
 
 	/** SCSEditor UI customizations */
 	TSharedPtr<ISCSEditorUICustomization> UICustomization;
@@ -653,6 +668,8 @@ protected:
 	 */
 	FSubobjectDataHandle CachedRootHandle;
 
+	UE_API void OnSelectionChanged(const UTypedElementSelectionSet* InChangedSelectionSet);
+
 	UE_API virtual FMenuBuilder CreateMenuBuilder();
 
 	/** Constructs the slate drag/drop tree for this subobject editor */
@@ -661,6 +678,9 @@ protected:
 	/** Creates a list of commands */
 	UE_API virtual void CreateCommandList();
 
+	/** Creates the selection set */
+	UE_API virtual void CreateSelectionSet();
+
 	/** Recursively visits the given node + its children and invokes the given function for each. */
 	UE_API void DepthFirstTraversal(const FSubobjectEditorTreeNodePtrType& InNodePtr, TSet<FSubobjectEditorTreeNodePtrType>& OutVisitedNodes, const TFunctionRef<void(const FSubobjectEditorTreeNodePtrType&)> InFunction) const;
 
diff --git a/Engine/Source/Editor/SubobjectEditor/SubobjectEditor.Build.cs b/Engine/Source/Editor/SubobjectEditor/SubobjectEditor.Build.cs
--- a/Engine/Source/Editor/SubobjectEditor/SubobjectEditor.Build.cs
+++ b/Engine/Source/Editor/SubobjectEditor/SubobjectEditor.Build.cs
@@ -13,7 +13,8 @@ public class SubobjectEditor : ModuleRules
 				"Core",
 				"CoreUObject",
 				"SubobjectDataInterface",
-				"Kismet",
+				"Kismet", 
+				"TypedElementRuntime",
 			}
 		);
 
@@ -33,6 +34,7 @@ public class SubobjectEditor : ModuleRules
 				"GraphEditor",
 				"ToolMenus",
 				"ToolWidgets",
+				"TypedElementFramework",
 			}
 		);
 		
