Blueprint Editor: Cleaned up some transaction session names to better match the level editor.

[RN]

diff --git a/Engine/Source/Editor/Kismet/Private/SCSEditorViewportClient.cpp b/Engine/Source/Editor/Kismet/Private/SCSEditorViewportClient.cpp
--- a/Engine/Source/Editor/Kismet/Private/SCSEditorViewportClient.cpp
+++ b/Engine/Source/Editor/Kismet/Private/SCSEditorViewportClient.cpp
@@ -395,7 +395,7 @@ void FSCSEditorViewportClient::ProcessClick(class FSceneView& View, class HHitPr
 		if (BlueprintEditorPtr.IsValid())
 		{
 			// If the user clicked on empty space, clear the selection
-			const FScopedTransaction Transaction(NSLOCTEXT("UnrealEd", "SelectNone", "Select None"));
+			const FScopedTransaction Transaction(NSLOCTEXT("UnrealEd", "ClickingBackground", "Clicking Background"));
 			BlueprintEditorPtr.Pin()->GetSubobjectEditor()->ClearSelection();
 		}
 		return;
diff --git a/Engine/Source/Editor/SubobjectEditor/Private/SSubobjectEditor.cpp b/Engine/Source/Editor/SubobjectEditor/Private/SSubobjectEditor.cpp
--- a/Engine/Source/Editor/SubobjectEditor/Private/SSubobjectEditor.cpp
+++ b/Engine/Source/Editor/SubobjectEditor/Private/SSubobjectEditor.cpp
@@ -2413,7 +2413,24 @@ void SSubobjectEditor::OnTreeSelectionChanged(FSubobjectEditorTreeNodePtrType In
 	// Don't update the tree based on updates to the selection
 	TGuardValue BlockTreeWidgetUpdates{ bBlockTreeWidgetUpdates, true };
 
-	const FScopedTransaction Transaction(NSLOCTEXT("UnrealEd", "SelectingElements", "Selecting Elements"), !bBlockSelectionUpdates && SelectInfo != ESelectInfo::Direct);
+	FText SessionName;
+	if (SelectInfo == ESelectInfo::OnMouseClick)
+	{
+		if (TreeWidget->GetNumItemsSelected() == 0)
+		{
+			SessionName = NSLOCTEXT("UnrealEd", "ClickingBackground", "Clicking Background");
+		}
+		else
+		{
+			SessionName = NSLOCTEXT("UnrealEd", "ClickingOnElements", "Clicking on Elements");
+		}
+	}
+	else
+	{
+		SessionName = NSLOCTEXT("UnrealEd", "SelectingElements", "Selecting Elements");
+	}
+
+	const FScopedTransaction Transaction(SessionName, !bBlockSelectionUpdates && SelectInfo != ESelectInfo::Direct);
 	UpdateSelectionFromNodes(TreeWidget->GetSelectedItems());
 }
 
