Blueprint Editor: Fix for crash when adding components.

[RN]

diff --git a/Engine/Source/Editor/SubobjectEditor/Private/SSubobjectEditor.cpp b/Engine/Source/Editor/SubobjectEditor/Private/SSubobjectEditor.cpp
--- a/Engine/Source/Editor/SubobjectEditor/Private/SSubobjectEditor.cpp
+++ b/Engine/Source/Editor/SubobjectEditor/Private/SSubobjectEditor.cpp
@@ -3048,6 +3048,9 @@ void SSubobjectEditor::OnRenameComponent()
 
 void SSubobjectEditor::OnRenameComponent(TUniquePtr<FScopedTransaction> InComponentCreateTransaction)
 {
+	// Make updates to the selection set are pushed to the tree before moving forward
+	SelectedElements->NotifyPendingChanges();
+
 	TArray<FSubobjectEditorTreeNodePtrType> SelectedNodes = GetSelectedNodes();
 
 	// Should already be prevented from making it here.
