Blueprint Editor: Duplicating elements will no longer cause focus to be automatically cleared.
Instance Subobject Editor: Duplicating elements will now select all duplicated nodes like the Blueprint version.

[RN]

diff --git a/Engine/Source/Editor/SubobjectEditor/Private/SSubobjectBlueprintEditor.cpp b/Engine/Source/Editor/SubobjectEditor/Private/SSubobjectBlueprintEditor.cpp
--- a/Engine/Source/Editor/SubobjectEditor/Private/SSubobjectBlueprintEditor.cpp
+++ b/Engine/Source/Editor/SubobjectEditor/Private/SSubobjectBlueprintEditor.cpp
@@ -28,6 +28,8 @@
 #include "Kismet2/ChildActorComponentEditorUtils.h"
 #include "ObjectTools.h"						// ThumbnailTools::CacheEmptyThumbnail
 #include "K2Node_ComponentBoundEvent.h"
+#include "Widgets/Input/SMultiLineEditableTextBox.h"
+#include "Widgets/Text/SInlineEditableTextBlock.h"
 
 #define LOCTEXT_NAMESPACE "SSubobjectBlueprintEditor"
 
@@ -293,18 +295,30 @@ void SSubobjectBlueprintEditor::OnDuplicateComponent()
 		return;
 	}
 
-	// Force the text box being edited (if any) to commit its text. The duplicate operation may trigger a regeneration of the tree view,
-	// releasing all row widgets. If one row was in edit mode (rename/rename on create), it was released before losing the focus and
-	// this would prevent the completion of the 'rename' or 'create + give initial name' transaction (occurring on focus lost).
-	FSlateApplication::Get().ClearKeyboardFocus();
-		
-	TUniquePtr<FScopedTransaction> Transaction = MakeUnique<FScopedTransaction>(SelectedNodes.Num() > 1 ? LOCTEXT("DuplicateComponents", "Duplicate Components") : LOCTEXT("DuplicateComponent", "Duplicate Component"));
 	USubobjectDataSubsystem* System = USubobjectDataSubsystem::Get();
 	if (!System)
 	{
 		return;
 	}
 
+	bool ClearedFocus = false;
+
+	// Check if the focus is currently in some sort of textbox
+	if (const TSharedPtr<SWidget> FocusedWidget = FSlateApplication::Get().GetKeyboardFocusedWidget())
+	{
+		if (ShouldClearFocusOnDuplicate(FocusedWidget))
+		{
+			// Force the text box being edited (if any) to commit its text. The duplicate operation may trigger a regeneration of the tree view,
+			// releasing all row widgets. If one row was in edit mode (rename/rename on create), it was released before losing the focus and
+			// this would prevent the completion of the 'rename' or 'create + give initial name' transaction (occurring on focus lost).
+			FSlateApplication::Get().ClearKeyboardFocus();
+
+			ClearedFocus = true;
+		}
+	}
+		
+	TUniquePtr<FScopedTransaction> Transaction = MakeUnique<FScopedTransaction>(SelectedNodes.Num() > 1 ? LOCTEXT("DuplicateComponents", "Duplicate Components") : LOCTEXT("DuplicateComponent", "Duplicate Component"));
+
 	TArray<FSubobjectDataHandle> DuplicatedHandles;
 	System->DuplicateSubobjects(GetObjectContextHandle(), SelectedNodes, GetBlueprint(), DuplicatedHandles);
 	UpdateTree();
@@ -342,6 +356,11 @@ void SSubobjectBlueprintEditor::OnDuplicateComponent()
 	else if (FirstNode)
 	{
 		TreeWidget->RequestScrollIntoView(FirstNode);
+		if (ClearedFocus)
+		{
+			// If focus had to be cleared, put it back on the tree widget so further interactions can be made with the new selection
+			FSlateApplication::Get().SetKeyboardFocus(TreeWidget, EFocusCause::SetDirectly);
+		}
 	}
 }
 
diff --git a/Engine/Source/Editor/SubobjectEditor/Private/SSubobjectEditor.cpp b/Engine/Source/Editor/SubobjectEditor/Private/SSubobjectEditor.cpp
--- a/Engine/Source/Editor/SubobjectEditor/Private/SSubobjectEditor.cpp
+++ b/Engine/Source/Editor/SubobjectEditor/Private/SSubobjectEditor.cpp
@@ -73,6 +73,7 @@ static const FName SubobjectTree_ColumnName_ComponentClass("ComponentClass");
 static const FName SubobjectTree_ColumnName_Asset("Asset");
 static const FName SubobjectTree_ColumnName_Mobility("Mobility"); 
 static const FName SubobjectTree_ContextMenuName("Kismet.SubobjectEditorContextMenu");
+static const FName SubobjectEditor_ClearFocusOnDuplicate("SubobjectEditor_ClearFocusOnDuplicate");
 
 ///////////////////////////////////////////////////////////
 // Util Functions
@@ -760,6 +761,7 @@ TSharedRef<SWidget> SSubobject_RowWidget::GenerateWidgetForColumn(const FName& C
 			.IsReadOnly(this, &SSubobject_RowWidget::IsReadOnly);
 
 		Node->SetRenameRequestedDelegate(FSubobjectEditorTreeNode::FOnRenameRequested::CreateSP(InlineWidget.Get(), &SInlineEditableTextBlock::EnterEditingMode));
+		InlineWidget->AddMetadata<FTagMetaData>(MakeShared<FTagMetaData>(SubobjectEditor_ClearFocusOnDuplicate));
 		
 		const bool IsRootActorNode = Data->IsRootActor();
 		
@@ -2859,6 +2861,25 @@ bool SSubobjectEditor::CanPasteNodes() const
 	return false;
 }
 
+bool SSubobjectEditor::ShouldClearFocusOnDuplicate(TSharedPtr<SWidget> InCurrentlyFocusedWidget)
+{
+	TSharedPtr<SWidget> CurrentWidget = InCurrentlyFocusedWidget;
+	while (CurrentWidget)
+	{
+		for (const TSharedRef<FTagMetaData>& TagMetaData : CurrentWidget->GetAllMetaData<FTagMetaData>())
+		{
+			if (TagMetaData->Tag == SubobjectEditor_ClearFocusOnDuplicate)
+			{
+				return true;
+			}
+		}
+
+		CurrentWidget = CurrentWidget->GetParentWidget();
+	}
+
+	return false;
+}
+
 bool SSubobjectEditor::CanDeleteNodes() const
 {
 	if(!IsEditingAllowed())
diff --git a/Engine/Source/Editor/SubobjectEditor/Private/SSubobjectInstanceEditor.cpp b/Engine/Source/Editor/SubobjectEditor/Private/SSubobjectInstanceEditor.cpp
--- a/Engine/Source/Editor/SubobjectEditor/Private/SSubobjectInstanceEditor.cpp
+++ b/Engine/Source/Editor/SubobjectEditor/Private/SSubobjectInstanceEditor.cpp
@@ -26,6 +26,7 @@
 #include "Widgets/Input/SSearchBox.h"
 #include "Widgets/Notifications/SNotificationList.h"
 #include "Widgets/SToolTip.h"
+#include "Widgets/Text/SInlineEditableTextBlock.h"
 
 #define LOCTEXT_NAMESPACE "SSubobjectInstanceEditor"
 
@@ -257,28 +258,76 @@ void SSubobjectInstanceEditor::CopySelectedNodes()
 void SSubobjectInstanceEditor::OnDuplicateComponent()
 {
 	TArray<FSubobjectDataHandle> SelectedNodes = GetSelectedHandles();
-	if(SelectedNodes.Num() > 0)
+	if(SelectedNodes.Num() <= 0)
 	{
-		// Force the text box being edited (if any) to commit its text. The duplicate operation may trigger a regeneration of the tree view,
-		// releasing all row widgets. If one row was in edit mode (rename/rename on create), it was released before losing the focus and
-		// this would prevent the completion of the 'rename' or 'create + give initial name' transaction (occurring on focus lost).
-		FSlateApplication::Get().ClearKeyboardFocus();
-		
-		TUniquePtr<FScopedTransaction> Transaction = MakeUnique<FScopedTransaction>(SelectedNodes.Num() > 1 ? LOCTEXT("DuplicateComponents", "Duplicate Components") : LOCTEXT("DuplicateComponent", "Duplicate Component"));
+		return;
+	}
+
+	USubobjectDataSubsystem* System = USubobjectDataSubsystem::Get();
+	if (!System)
+	{
+		return;
+	}
 
-		if(USubobjectDataSubsystem* System = USubobjectDataSubsystem::Get())
+	bool ClearedFocus = false;
+
+	// Check if the focus is currently in some sort of textbox
+	if (const TSharedPtr<SWidget> FocusedWidget = FSlateApplication::Get().GetKeyboardFocusedWidget())
+	{
+		if (ShouldClearFocusOnDuplicate(FocusedWidget))
 		{
-			TArray<FSubobjectDataHandle> DuplicatedHandles;
-			System->DuplicateSubobjects(GetObjectContextHandle(), SelectedNodes, /* BpContext = */ nullptr, DuplicatedHandles);
-			UpdateTree();
+			// Force the text box being edited (if any) to commit its text. The duplicate operation may trigger a regeneration of the tree view,
+			// releasing all row widgets. If one row was in edit mode (rename/rename on create), it was released before losing the focus and
+			// this would prevent the completion of the 'rename' or 'create + give initial name' transaction (occurring on focus lost).
+			FSlateApplication::Get().ClearKeyboardFocus();
+
+			ClearedFocus = true;
+		}
+	}
+		
+	TUniquePtr<FScopedTransaction> Transaction = MakeUnique<FScopedTransaction>(SelectedNodes.Num() > 1 ? LOCTEXT("DuplicateComponents", "Duplicate Components") : LOCTEXT("DuplicateComponent", "Duplicate Component"));
+
+	TArray<FSubobjectDataHandle> DuplicatedHandles;
+	System->DuplicateSubobjects(GetObjectContextHandle(), SelectedNodes, /* BpContext = */ nullptr, DuplicatedHandles);
+	UpdateTree();
 
-			// Set focus to the newly created subobject
-			FSubobjectEditorTreeNodePtrType NewNode = DuplicatedHandles.Num() > 0 ? FindSlateNodeForHandle(DuplicatedHandles[0]) : nullptr;
-			if (NewNode != nullptr)
+	// If nothing was duplicated, no further action is needed
+	if (DuplicatedHandles.Num() <= 0)
+	{
+		return;
+	}
+
+	ClearSelection();
+
+	FSubobjectEditorTreeNodePtrType FirstNode = nullptr;
+
+	// Select all duplicated subobjects
+	for (const FSubobjectDataHandle& DataHandle : DuplicatedHandles)
+	{
+		if (const FSubobjectEditorTreeNodePtrType NewNode = FindSlateNodeForHandle(DataHandle))
+		{
+			if (!FirstNode)
 			{
-				TreeWidget->SetSelection(NewNode);
-				OnRenameComponent(MoveTemp(Transaction));
+				FirstNode = NewNode;
 			}
+
+			SelectNode(NewNode, true);
+		}
+	}
+
+	// If only one subobject was duplicated, start a rename operation on it
+	if (DuplicatedHandles.Num() == 1)
+	{
+		OnRenameComponent(MoveTemp(Transaction));
+	}
+	// Otherwise, scroll the first of the duplicated nodes into view in the tree widget
+	else if (FirstNode)
+	{
+		TreeWidget->RequestScrollIntoView(FirstNode);
+		if (ClearedFocus)
+		{
+			// If focus had to be cleared, put it back on the tree widget so further interactions can be made with the new selection
+			FSlateApplication::Get().SetKeyboardFocus(TreeWidget, EFocusCause::SetDirectly);
 		}
 	}
 }
diff --git a/Engine/Source/Editor/SubobjectEditor/Public/SSubobjectEditor.h b/Engine/Source/Editor/SubobjectEditor/Public/SSubobjectEditor.h
--- a/Engine/Source/Editor/SubobjectEditor/Public/SSubobjectEditor.h
+++ b/Engine/Source/Editor/SubobjectEditor/Public/SSubobjectEditor.h
@@ -586,6 +586,9 @@ protected:
 	/** Pastes previously copied node(s) */
 	UE_API virtual bool CanPasteNodes() const;
 	virtual void PasteNodes() = 0;
+
+	/** Returns true if keyboard focus should be cleared in OnDuplicateComponent if the widget is currently focused */
+	UE_API static bool ShouldClearFocusOnDuplicate(TSharedPtr<SWidget> InCurrentlyFocusedWidget);
 	
 public:
 	/** Removes existing selected component nodes from the SCS */
