Fix crash that could occur in the SubobjectEditor when hovering Add Event submenu in a subobject context menu
* GetBlueprint can return nullptr if SSubobjectBlueprintEditor is being used with an editor that is not the blueprint editor.

diff --git a/Engine/Source/Editor/SubobjectEditor/Private/SSubobjectBlueprintEditor.cpp b/Engine/Source/Editor/SubobjectEditor/Private/SSubobjectBlueprintEditor.cpp
--- a/Engine/Source/Editor/SubobjectEditor/Private/SSubobjectBlueprintEditor.cpp
+++ b/Engine/Source/Editor/SubobjectEditor/Private/SSubobjectBlueprintEditor.cpp
@@ -621,15 +621,18 @@ void SSubobjectBlueprintEditor::PopulateContextMenuImpl(UToolMenu* InMenu, TArra
 			// Find the common base class of all selected classes
 			UClass* SelectedClass = UClass::FindCommonBase(SelectionClasses);
 			// Build an event submenu if we can generate events
-			if (FBlueprintEditorUtils::CanClassGenerateEvents(SelectedClass))
+			UBlueprint* Blueprint = GetBlueprint();
+			if (FBlueprintEditorUtils::CanClassGenerateEvents(SelectedClass) && Blueprint != nullptr)
 			{
 				BlueprintSCSSection.AddSubMenu(
-                    "AddEventSubMenu",
-                    LOCTEXT("AddEventSubMenu", "Add Event"),
-                    LOCTEXT("ActtionsSubMenu_ToolTip", "Add Event"),
-                    FNewMenuDelegate::CreateStatic(&SSubobjectBlueprintEditor::BuildMenuEventsSection,
-                        GetBlueprint(), SelectedClass, FCanExecuteAction::CreateSP(this, &SSubobjectEditor::IsEditingAllowed),
-                        FGetSelectedObjectsDelegate::CreateSP(this, &SSubobjectEditor::GetSelectedItemsForContextMenu)));
+					"AddEventSubMenu",
+					LOCTEXT("AddEventSubMenu", "Add Event"),
+					LOCTEXT("ActtionsSubMenu_ToolTip", "Add Event"),
+					FNewMenuDelegate::CreateStatic(&SSubobjectBlueprintEditor::BuildMenuEventsSection,
+						Blueprint,
+						SelectedClass,
+						FCanExecuteAction::CreateSP(this, &SSubobjectEditor::IsEditingAllowed),
+						FGetSelectedObjectsDelegate::CreateSP(this, &SSubobjectEditor::GetSelectedItemsForContextMenu)));
 			}
 		}
 	}
