Editor: Fix FSubobjectData and FSubobjectEditorTreeNode being leaked due to shared pointer cycles.

diff --git a/Engine/Source/Editor/SubobjectDataInterface/Private/SubobjectData.cpp b/Engine/Source/Editor/SubobjectDataInterface/Private/SubobjectData.cpp
--- a/Engine/Source/Editor/SubobjectDataInterface/Private/SubobjectData.cpp
+++ b/Engine/Source/Editor/SubobjectDataInterface/Private/SubobjectData.cpp
@@ -31,7 +31,6 @@ PRAGMA_ENABLE_DEPRECATION_WARNINGS
 
 FSubobjectData::FSubobjectData()
 	: WeakObjectPtr(nullptr)
-	, ParentObjectHandle(FSubobjectDataHandle::InvalidHandle)
 	, SCSNodePtr(nullptr)
 {
 	// By Default the handle will be invalid, it will only
@@ -40,7 +39,6 @@ FSubobjectData::FSubobjectData()
 
 FSubobjectData::FSubobjectData(UObject* ContextObject, const FSubobjectDataHandle& InParentHandle, const bool bInIsInheritedSCS)
 	: WeakObjectPtr(ContextObject)
-	, ParentObjectHandle(InParentHandle)
 	, SCSNodePtr(nullptr)
 {
 	FSubobjectDataSubsystemContextData_SingleUObjectContextObject SingleUObject(ContextObject);
@@ -89,7 +87,7 @@ void FSubobjectData::Initialize(TInstancedStruct<FSubobjectDataSubsystemContextD
 	}
 
 	ContextData = MoveTemp(InContextData);
-	ParentObjectHandle = ParentHandle;
+	ParentObjectHandle = ParentHandle.GetSharedDataPtr();
 }
 
 bool FSubobjectData::CanEdit() const
@@ -671,7 +669,7 @@ void FSubobjectData::SetupAttachment(FName SocketName, const FSubobjectDataHandl
 
 FSubobjectDataHandle FSubobjectData::GetRootSubobject() const
 {
-	FSubobjectDataHandle Current = ParentObjectHandle;
+	FSubobjectDataHandle Current = ParentObjectHandle.Pin();
 	while(Current.IsValid() && Current.GetSharedDataPtr()->HasParent())
 	{
 		Current = Current.GetSharedDataPtr()->GetParentHandle();
@@ -731,7 +729,7 @@ FSubobjectDataHandle FSubobjectData::FindChildByObject(UObject* ContextObject) c
 bool FSubobjectData::AddChildHandleOnly(const FSubobjectDataHandle& InHandle)
 {
 	// If we already have this child, then don't both with adding it
-	if (HasChild(InHandle) || !InHandle.IsValid() || InHandle == Handle)
+	if (HasChild(InHandle) || !InHandle.IsValid() || InHandle == Handle.Pin())
 	{
 		return false;
 	}
@@ -739,7 +737,7 @@ bool FSubobjectData::AddChildHandleOnly(const FSubobjectDataHandle& InHandle)
 	ChildrenHandles.Add(InHandle);
 	if(FSubobjectData* NewChildData = InHandle.GetData())
 	{
-		NewChildData->SetParentHandle(Handle);		
+		NewChildData->SetParentHandle(Handle.Pin());		
 	}
 
 	return true;
@@ -1198,7 +1196,7 @@ bool FSubobjectData::IsInstancedInheritedComponent() const
 	{
 		return false;
 	}
-	FSubobjectDataHandle CurrentHandle = ParentObjectHandle;
+	FSubobjectDataHandle CurrentHandle = ParentObjectHandle.Pin();
 	FSubobjectData* CurrentData = CurrentHandle.GetData();
 			
 	while(CurrentHandle.IsValid() && CurrentData && !CurrentData->IsActor())
@@ -1214,7 +1212,7 @@ bool FSubobjectData::IsAttachedTo(const FSubobjectDataHandle& InHandle) const
 {
 	REDIRECT_TO_ADAPTER(IsAttachedTo, InHandle)
 
-	FSubobjectDataHandle TestParentHandle = ParentObjectHandle;
+	FSubobjectDataHandle TestParentHandle = ParentObjectHandle.Pin();
 
 	while (TestParentHandle.IsValid())
 	{
diff --git a/Engine/Source/Editor/SubobjectDataInterface/Private/SubobjectDataSubsystem.cpp b/Engine/Source/Editor/SubobjectDataInterface/Private/SubobjectDataSubsystem.cpp
--- a/Engine/Source/Editor/SubobjectDataInterface/Private/SubobjectDataSubsystem.cpp
+++ b/Engine/Source/Editor/SubobjectDataInterface/Private/SubobjectDataSubsystem.cpp
@@ -2804,7 +2804,7 @@ FScopedTransaction* USubobjectDataSubsystem::BeginTransaction(const TArray<FSubo
 FSubobjectDataHandle USubobjectDataSubsystem::CreateSubobjectData()
 {
 	TSharedPtr<FSubobjectData> SharedPtr = MakeShared<FSubobjectData>();
-	SharedPtr->Handle.DataPtr = SharedPtr;
+	SharedPtr->Handle = SharedPtr;
 	return SharedPtr->GetHandle();
 }
 
@@ -2824,7 +2824,7 @@ FSubobjectDataHandle USubobjectDataSubsystem::CreateSubobjectData(UObject* Conte
 
 	check(SharedPtr.IsValid());
 	
-	SharedPtr->Handle.DataPtr = SharedPtr;
+	SharedPtr->Handle = SharedPtr;
 
 	return SharedPtr->GetHandle();
 }
diff --git a/Engine/Source/Editor/SubobjectDataInterface/Public/SubobjectData.h b/Engine/Source/Editor/SubobjectDataInterface/Public/SubobjectData.h
--- a/Engine/Source/Editor/SubobjectDataInterface/Public/SubobjectData.h
+++ b/Engine/Source/Editor/SubobjectDataInterface/Public/SubobjectData.h
@@ -179,9 +179,9 @@ public:
 	UE_API void SetSocketName(FName InNewName);
 	UE_API void SetupAttachment(FName SocketName, const FSubobjectDataHandle& AttachParentHandle = FSubobjectDataHandle::InvalidHandle);
 
-	FSubobjectDataHandle GetHandle() const { return Handle; }
-	FSubobjectDataHandle GetParentHandle() const { return ParentObjectHandle; }
-	FSubobjectData* GetParentData() const { return ParentObjectHandle.GetData(); }
+	FSubobjectDataHandle GetHandle() const { return Handle.Pin(); }
+	FSubobjectDataHandle GetParentHandle() const { return ParentObjectHandle.Pin(); }
+	FSubobjectData* GetParentData() const { return ParentObjectHandle.Pin().Get(); }
 	bool HasParent() const { return ParentObjectHandle.IsValid(); }
 	const TArray<FSubobjectDataHandle>& GetChildrenHandles() const { return ChildrenHandles; }
 	
@@ -244,20 +244,18 @@ private:
 	*/
 	UE_API bool RemoveChildHandleOnly(const FSubobjectDataHandle& InHandle);
 
-	void SetParentHandle(const FSubobjectDataHandle& NewParentHandle) { ParentObjectHandle = NewParentHandle; }
-	void ClearParentHandle() { ParentObjectHandle = FSubobjectDataHandle::InvalidHandle; }
+	void SetParentHandle(const FSubobjectDataHandle& NewParentHandle) { ParentObjectHandle = NewParentHandle.GetSharedDataPtr(); }
+	void ClearParentHandle() { ParentObjectHandle.Reset(); }
 	
 	// Weak ptr to the object instance represented by this node (e.g. component template)
 	UPROPERTY()
 	TWeakObjectPtr<UObject> WeakObjectPtr;
 
-	// Handle that will have a globally unique ID when it is constructed (i.e. when this is constructed)
-	UPROPERTY()
-	FSubobjectDataHandle Handle;
+	// Handle to this object
+	TWeakPtr<FSubobjectData> Handle;
 
 	// Handle to this object's parent. If this is invalid, then it is the root subobject.
-	UPROPERTY()
-	FSubobjectDataHandle ParentObjectHandle;
+	TWeakPtr<FSubobjectData> ParentObjectHandle;
 
 	// Set of child subobjects that use this subobject as a parent
 	UPROPERTY()
diff --git a/Engine/Source/Editor/SubobjectDataInterface/Public/SubobjectDataHandle.h b/Engine/Source/Editor/SubobjectDataInterface/Public/SubobjectDataHandle.h
--- a/Engine/Source/Editor/SubobjectDataInterface/Public/SubobjectDataHandle.h
+++ b/Engine/Source/Editor/SubobjectDataInterface/Public/SubobjectDataHandle.h
@@ -16,10 +16,8 @@
 struct FSubobjectData;
 
 /**
-* A subobject handle is a globally unique identifier for subobjects
-* Upon construction, the handle will be invalid. It is the responsibility
-* of the owning FSubobjectData to set the DataPtr once the subobject
-* data has validated that it has a good context.
+* A subobject handle is a wrapper around a TSharedPtr<FSubobjectData> intended for use
+* with Blueprints.
 */
 USTRUCT(BlueprintType)
 struct FSubobjectDataHandle
@@ -32,6 +30,9 @@ struct FSubobjectDataHandle
 	explicit FSubobjectDataHandle()
 		: DataPtr(nullptr)
 	{}
+	FSubobjectDataHandle(const TSharedPtr<FSubobjectData> Data)
+		: DataPtr(Data)
+	{}
 	
 	~FSubobjectDataHandle()
 	{
diff --git a/Engine/Source/Editor/SubobjectEditor/Private/SSubobjectEditor.cpp b/Engine/Source/Editor/SubobjectEditor/Private/SSubobjectEditor.cpp
--- a/Engine/Source/Editor/SubobjectEditor/Private/SSubobjectEditor.cpp
+++ b/Engine/Source/Editor/SubobjectEditor/Private/SSubobjectEditor.cpp
@@ -195,7 +195,7 @@ bool FSubobjectEditorTreeNode::IsNativeComponent() const
 bool FSubobjectEditorTreeNode::IsRootActorNode() const
 {
 	// If there is no parent node pointer then this has to be top of the tree
-	return ParentNodePtr == nullptr;
+	return !Parent.IsValid();
 }
 
 FName FSubobjectEditorTreeNode::GetVariableName() const
@@ -235,7 +235,7 @@ bool FSubobjectEditorTreeNode::IsChildSubtreeNode() const
 
 bool FSubobjectEditorTreeNode::IsAttachedTo(FSubobjectEditorTreeNodePtrType InNodePtr) const
 {
-	FSubobjectEditorTreeNodePtrType TestParentPtr = ParentNodePtr;
+	FSubobjectEditorTreeNodePtrType TestParentPtr = Parent.Pin();
 	while (TestParentPtr.IsValid())
 	{
 		if (TestParentPtr == InNodePtr)
@@ -243,7 +243,7 @@ bool FSubobjectEditorTreeNode::IsAttachedTo(FSubobjectEditorTreeNodePtrType InNo
 			return true;
 		}
 
-		TestParentPtr = TestParentPtr->ParentNodePtr;
+		TestParentPtr = TestParentPtr->Parent.Pin();
 	}
 
 	return false;
@@ -251,7 +251,7 @@ bool FSubobjectEditorTreeNode::IsAttachedTo(FSubobjectEditorTreeNodePtrType InNo
 
 bool FSubobjectEditorTreeNode::IsDirectlyAttachedTo(FSubobjectEditorTreeNodePtrType InNodePtr) const
 {
-	return ParentNodePtr == InNodePtr;
+	return Parent.Pin() == InNodePtr;
 }
 
 bool FSubobjectEditorTreeNode::CanDelete() const
@@ -283,13 +283,13 @@ void FSubobjectEditorTreeNode::AddChild(FSubobjectEditorTreeNodePtrType AttachTo
 	}
 	
 	Children.AddUnique(AttachToPtr);
-	AttachToPtr->ParentNodePtr = AsShared();
+	AttachToPtr->Parent = AsShared();
 }
 
 void FSubobjectEditorTreeNode::RemoveChild(FSubobjectEditorTreeNodePtrType InChildNodePtr)
 {
 	Children.Remove(InChildNodePtr);
-	InChildNodePtr->ParentNodePtr.Reset();
+	InChildNodePtr->Parent.Reset();
 }
 
 FSubobjectEditorTreeNodePtrType FSubobjectEditorTreeNode::FindChild(const FSubobjectDataHandle& InHandle)
@@ -344,15 +344,15 @@ void FSubobjectEditorTreeNode::SetCachedFilterState(bool bMatchesFilter, bool bU
 void FSubobjectEditorTreeNode::ApplyFilteredStateToParent()
 {
 	FSubobjectEditorTreeNode* Child = this;
-	while (Child->ParentNodePtr.IsValid())
+	while (Child->Parent.IsValid())
 	{
-		FSubobjectEditorTreeNode* Parent = Child->ParentNodePtr.Get();
+		FSubobjectEditorTreeNode* CurrentParent = Child->Parent.Pin().Get();
 
 		if (!IsFlaggedForFiltration())
 		{
-			if ((Parent->FilterFlags & EFilteredState::ChildMatches) == 0)
+			if ((CurrentParent->FilterFlags & EFilteredState::ChildMatches) == 0)
 			{
-				Parent->FilterFlags |= EFilteredState::ChildMatches;
+				CurrentParent->FilterFlags |= EFilteredState::ChildMatches;
 			}
 			else
 			{
@@ -361,10 +361,10 @@ void FSubobjectEditorTreeNode::ApplyFilteredStateToParent()
 			}
 		}
 		// have to see if this was the only child contributing to this flag
-		else if (Parent->FilterFlags & EFilteredState::ChildMatches)
+		else if (CurrentParent->FilterFlags & EFilteredState::ChildMatches)
 		{
-			Parent->FilterFlags &= ~EFilteredState::ChildMatches;
-			for (const FSubobjectEditorTreeNodePtrType& Sibling : Parent->Children)
+			CurrentParent->FilterFlags &= ~EFilteredState::ChildMatches;
+			for (const FSubobjectEditorTreeNodePtrType& Sibling : CurrentParent->Children)
 			{
 				if (Sibling.Get() == Child)
 				{
@@ -373,18 +373,18 @@ void FSubobjectEditorTreeNode::ApplyFilteredStateToParent()
 
 				if (Sibling->FilterFlags & EFilteredState::FilteredInMask)
 				{
-					Parent->FilterFlags |= EFilteredState::ChildMatches;
+					CurrentParent->FilterFlags |= EFilteredState::ChildMatches;
 					break;
 				}
 			}
 
-			if (Parent->FilterFlags & EFilteredState::ChildMatches)
+			if (CurrentParent->FilterFlags & EFilteredState::ChildMatches)
 			{
 				// another child added the flag back
 				break;
 			}
 		}
-		Child = Parent;
+		Child = CurrentParent;
 	}
 }
 
diff --git a/Engine/Source/Editor/SubobjectEditor/Public/SSubobjectEditor.h b/Engine/Source/Editor/SubobjectEditor/Public/SSubobjectEditor.h
--- a/Engine/Source/Editor/SubobjectEditor/Public/SSubobjectEditor.h
+++ b/Engine/Source/Editor/SubobjectEditor/Public/SSubobjectEditor.h
@@ -49,7 +49,9 @@ public:
 	
 	UE_API explicit FSubobjectEditorTreeNode(const FSubobjectDataHandle& DataSource, bool InbIsSeperator = false);
 
+	PRAGMA_DISABLE_DEPRECATION_WARNINGS // complains about ParentNodePtr otherwise
 	~FSubobjectEditorTreeNode() = default;
+	PRAGMA_ENABLE_DEPRECATION_WARNINGS
 
 	bool IsSeperator() const { return bIsSeperator; } 
 
@@ -66,7 +68,7 @@ public:
 	FSubobjectDataHandle GetDataHandle() const { return DataHandle; }
 
 	/** Get the slate parent of this single tree node */
-	FSubobjectEditorTreeNodePtrType GetParent() const { return ParentNodePtr; }
+	FSubobjectEditorTreeNodePtrType GetParent() const { return Parent.Pin(); }
 
 	/** Get the display name of this single tree node */
 	UE_API FString GetDisplayString() const;
@@ -144,9 +146,11 @@ public:
 
 protected: 
 
-	/** Pointer to the parent of this subobject */
+	UE_DEPRECATED(5.7, "Use the Parent weak pointer instead.")
 	FSubobjectEditorTreeNodePtrType ParentNodePtr;
 	
+	TWeakPtr<FSubobjectEditorTreeNode> Parent;
+	
 	// Scope the creation of a node which ends when the initial 'name' is given/accepted by the user, which can be several frames after the node was actually created.
 	TUniquePtr<FScopedTransaction> OngoingCreateTransaction;
 
