Enable optional overlay over the tree view widget in the subobject editor

diff --git a/Engine/Source/Editor/Kismet/Public/ISCSEditorUICustomization.h b/Engine/Source/Editor/Kismet/Public/ISCSEditorUICustomization.h
--- a/Engine/Source/Editor/Kismet/Public/ISCSEditorUICustomization.h
+++ b/Engine/Source/Editor/Kismet/Public/ISCSEditorUICustomization.h
@@ -6,6 +6,7 @@
 #include "CoreMinimal.h"
 #include "Templates/SubclassOf.h"
 #include "Containers/ArrayView.h"
+#include "Widgets/SNullWidget.h"
 
 class UActorComponent;
 struct FSubobjectData;
@@ -78,4 +79,10 @@ public:
 	/** Optionally sorts the given SubobjectDataHandles for the tree view as appropriate for this editor.
 	 * @return Whether any sorting was attempted */
 	virtual bool SortSubobjectData(TArray<FSubobjectDataHandle>& SubobjectData) { return false; }
+
+	/** Decides if this UI Customization provides an overlay for the tree view */
+	virtual bool SupportsTreeViewOverlay(TArrayView<UObject*> Context) const { return false; }
+
+	/** Create an overlay for the tree view widget */
+	virtual TSharedRef<class SWidget> ConstructTreeViewOverlay(TWeakPtr<class SSubobjectEditor> Editor) const { return SNullWidget::NullWidget; }
 };
diff --git a/Engine/Source/Editor/LevelEditor/Private/ActorDetailsSCSEditorUICustomization.cpp b/Engine/Source/Editor/LevelEditor/Private/ActorDetailsSCSEditorUICustomization.cpp
--- a/Engine/Source/Editor/LevelEditor/Private/ActorDetailsSCSEditorUICustomization.cpp
+++ b/Engine/Source/Editor/LevelEditor/Private/ActorDetailsSCSEditorUICustomization.cpp
@@ -135,4 +135,28 @@ void FActorDetailsSCSEditorUICustomization::RemoveCustomization(TSharedPtr<ISCSE
 	Customizations.Remove(Customization);
 }
 
+bool FActorDetailsSCSEditorUICustomization::SupportsTreeViewOverlay(TArrayView<UObject*> Context) const
+{
+	for (const TSharedPtr<ISCSEditorUICustomization>& Customization : Customizations)
+	{
+		if (Customization->SupportsTreeViewOverlay(Context))
+		{
+			return true;
+		}
+	}
+	return false;
+}
 
+TSharedRef<SWidget> FActorDetailsSCSEditorUICustomization::ConstructTreeViewOverlay(TWeakPtr<SSubobjectEditor> Editor) const
+{
+	for (const TSharedPtr<ISCSEditorUICustomization>& Customization : Customizations)
+	{
+		const TSharedRef<SWidget> Overlay = Customization->ConstructTreeViewOverlay(Editor);
+		if (Overlay != SNullWidget::NullWidget)
+		{
+			return Overlay;
+		}
+	}
+
+	return SNullWidget::NullWidget;
+}
diff --git a/Engine/Source/Editor/LevelEditor/Private/ActorDetailsSCSEditorUICustomization.h b/Engine/Source/Editor/LevelEditor/Private/ActorDetailsSCSEditorUICustomization.h
--- a/Engine/Source/Editor/LevelEditor/Private/ActorDetailsSCSEditorUICustomization.h
+++ b/Engine/Source/Editor/LevelEditor/Private/ActorDetailsSCSEditorUICustomization.h
@@ -30,6 +30,9 @@ public:
 	
 	virtual bool FilterSubobjectData(TArray<FSubobjectDataHandle>& SubobjectData) const override;
 
+	virtual bool SupportsTreeViewOverlay(TArrayView<UObject*> Context) const override;
+	virtual TSharedRef<SWidget> ConstructTreeViewOverlay(TWeakPtr<SSubobjectEditor> Editor) const override;
+
 	void AddCustomization(TSharedPtr<ISCSEditorUICustomization> Customization);
 	void RemoveCustomization(TSharedPtr<ISCSEditorUICustomization> Customization);
 
diff --git a/Engine/Source/Editor/SubobjectEditor/Private/SSubobjectInstanceEditor.cpp b/Engine/Source/Editor/SubobjectEditor/Private/SSubobjectInstanceEditor.cpp
--- a/Engine/Source/Editor/SubobjectEditor/Private/SSubobjectInstanceEditor.cpp
+++ b/Engine/Source/Editor/SubobjectEditor/Private/SSubobjectInstanceEditor.cpp
@@ -161,13 +161,41 @@ void SSubobjectInstanceEditor::Construct(const FArguments& InArgs)
 
 	+ SVerticalBox::Slot()
 	[
-		SNew(SBorder)
-		.BorderImage(FAppStyle::Get().GetBrush("SCSEditor.Background"))
-		.Padding(4.f)
-		.AddMetaData<FTagMetaData>(FTagMetaData(TEXT("ComponentsPanel")))
-		.Visibility(this, &SSubobjectInstanceEditor::GetComponentsTreeVisibility)
+		SNew(SOverlay)
+		+ SOverlay::Slot()
 		[
-			TreeWidget.ToSharedRef()
+			SNew(SBorder)
+			.BorderImage(FAppStyle::Get().GetBrush("SCSEditor.Background"))
+			.Padding(4.f)
+			.AddMetaData<FTagMetaData>(FTagMetaData(TEXT("ComponentsPanel")))
+			.Visibility(this, &SSubobjectInstanceEditor::GetComponentsTreeVisibility)
+			[
+				TreeWidget.ToSharedRef()
+			]
+		]
+		+ SOverlay::Slot()
+		[
+			SAssignNew(TreeViewOverlay, SBox)
+			.Visibility_Lambda([this]()
+			{
+				const bool bTreeIsVisible = GetComponentsTreeVisibility() == EVisibility::Visible;
+				UObject* ObjectContextPtr = GetObjectContext();
+				const bool bSupportsOverlay = UICustomization->SupportsTreeViewOverlay(MakeArrayView(&ObjectContextPtr, 1));		
+				
+				if (bTreeIsVisible && bSupportsOverlay)
+				{
+					// TODO: Could add more criteria like whether the thing we're dragging can actually be dropped here
+					const bool bIsDragging = FSlateApplication::Get().IsDragDropping();
+					return bIsDragging ? EVisibility::Visible : EVisibility::HitTestInvisible;
+				} 
+
+				return EVisibility::Collapsed;
+			})
+			.Content()
+			[
+				UICustomization->ConstructTreeViewOverlay(SharedThis(this))
+			]
+			
 		]
 	];
 
@@ -200,6 +228,21 @@ void SSubobjectInstanceEditor::Construct(const FArguments& InArgs)
 }
 END_SLATE_FUNCTION_BUILD_OPTIMIZATION
 
+void SSubobjectInstanceEditor::UpdateTree(bool bRegenerateTreeNodes)
+{
+	SSubobjectEditor::UpdateTree(bRegenerateTreeNodes);
+
+	if (TreeViewOverlay.IsValid())
+	{
+		UObject* ObjectContextPtr = GetObjectContext();
+		if (UICustomization->SupportsTreeViewOverlay(MakeArrayView(&ObjectContextPtr, 1)))
+		{
+			TreeViewOverlay->SetContent(UICustomization->ConstructTreeViewOverlay(SharedThis(this)));
+		}
+	}
+}
+
+
 void SSubobjectInstanceEditor::OnDeleteNodes()
 {
 	const FScopedTransaction Transaction( LOCTEXT("RemoveComponents", "Remove Components") );
diff --git a/Engine/Source/Editor/SubobjectEditor/Public/SSubobjectEditor.h b/Engine/Source/Editor/SubobjectEditor/Public/SSubobjectEditor.h
--- a/Engine/Source/Editor/SubobjectEditor/Public/SSubobjectEditor.h
+++ b/Engine/Source/Editor/SubobjectEditor/Public/SSubobjectEditor.h
@@ -475,7 +475,7 @@ public:
 	UE_API void SelectNodeFromHandle(const FSubobjectDataHandle& InHandle, bool bIsCntrlDown);
 	
 	/** Update the contents of the Subobject Tree based on the current context */
-	UE_API void UpdateTree(bool bRegenerateTreeNodes = true);
+	UE_API virtual void UpdateTree(bool bRegenerateTreeNodes = true);
 
 	/** Dumps out the tree view contents to the log (used to assist with debugging widget hierarchy issues) */
 	UE_API void DumpTree();
diff --git a/Engine/Source/Editor/SubobjectEditor/Public/SSubobjectInstanceEditor.h b/Engine/Source/Editor/SubobjectEditor/Public/SSubobjectInstanceEditor.h
--- a/Engine/Source/Editor/SubobjectEditor/Public/SSubobjectInstanceEditor.h
+++ b/Engine/Source/Editor/SubobjectEditor/Public/SSubobjectInstanceEditor.h
@@ -42,6 +42,8 @@ protected:
     UE_API virtual void OnMakeNewRootDropAction(FSubobjectEditorTreeNodePtrType DroppedNodePtr) override;
     UE_API virtual void PostDragDropAction(bool bRegenerateTreeNodes) override;
 
+	UE_API virtual void UpdateTree(bool bRegenerateTreeNodes = true) override;
+
     /** Builds a context menu popup for dropping a child node onto the scene root node */
     UE_API virtual TSharedPtr<SWidget> BuildSceneRootDropActionMenu(FSubobjectEditorTreeNodePtrType DroppedOntoNodePtr, FSubobjectEditorTreeNodePtrType DroppedNodePtr) override;
 	UE_API virtual FSubobjectDataHandle AddNewSubobject(const FSubobjectDataHandle& ParentHandle, UClass* NewClass, UObject* AssetOverride, FText& OutFailReason, TUniquePtr<FScopedTransaction> InOngoingTransaction) override;
@@ -64,6 +66,8 @@ private:
 	/** @return the tooltip describing how many properties will be reset to the blueprint default*/
 	UE_API FText OnGetResetToBlueprintDefaultsTooltip() const;
 
+	TSharedPtr<class SBox> TreeViewOverlay;
+
 };
 
 #undef UE_API
