Editor: Enable multi-selection on non-actor items in subobject editor:

- Allow subobject editor visibility to be customized
- Use a FScopedRow inside FSubobjectDataSubsystemContextData_TedsRow.
- Add adapter function to potentially unpack multiple objects from a FSubobjectData
- Allow GatherSubobjectData to get selection directly from TEDS, by passing down a selection set identifier.

diff --git a/Engine/Source/Editor/Kismet/Public/ISCSEditorUICustomization.h b/Engine/Source/Editor/Kismet/Public/ISCSEditorUICustomization.h
--- a/Engine/Source/Editor/Kismet/Public/ISCSEditorUICustomization.h
+++ b/Engine/Source/Editor/Kismet/Public/ISCSEditorUICustomization.h
@@ -6,11 +6,16 @@
 #include "CoreMinimal.h"
 #include "Templates/SubclassOf.h"
 #include "Containers/ArrayView.h"
+#include "Layout/Visibility.h"
 #include "Widgets/SNullWidget.h"
 
 class UActorComponent;
 struct FSubobjectData;
 struct FSubobjectDataHandle;
+namespace UE::Editor
+{
+	struct FSubobjectEditorContext;
+}
 
 /** SCSEditor UI customization */
 class ISCSEditorUICustomization
@@ -62,6 +67,15 @@ public:
 	 */
 	virtual TSharedPtr<SWidget> GetControlsWidget(const FSubobjectData&) const { return {}; }
 
+	/**
+	 * When it returns true, overrides the value returned by SSubobjectEditor::GetDesiredVisibility().
+	 *  Otherwise, it is ignored.
+	 * @param VisibilityOut The visibility to return
+	 * @return true if the customization wants VisibilityOut to be used. 
+	 */
+	UE_EXPERIMENTAL(5.8, "Signature may change in the future")
+	virtual bool OverrideDesiredVisibility(const UE::Editor::FSubobjectEditorContext& Context, EVisibility& VisibilityOut) const { return false; }
+
 	/** 
 	 * @return An override for the default ChildActorComponentTreeViewVisualizationMode from the project settings, if different from UseDefault.
 	 * @note Setting an override also forces child actor tree view expansion to be enabled.
diff --git a/Engine/Source/Editor/LevelEditor/Private/ActorDetailsSCSEditorUICustomization.cpp b/Engine/Source/Editor/LevelEditor/Private/ActorDetailsSCSEditorUICustomization.cpp
--- a/Engine/Source/Editor/LevelEditor/Private/ActorDetailsSCSEditorUICustomization.cpp
+++ b/Engine/Source/Editor/LevelEditor/Private/ActorDetailsSCSEditorUICustomization.cpp
@@ -88,6 +88,18 @@ TSharedPtr<SWidget> FActorDetailsSCSEditorUICustomization::GetControlsWidget(con
 	return ISCSEditorUICustomization::GetControlsWidget(SubobjectData);
 }
 
+bool FActorDetailsSCSEditorUICustomization::OverrideDesiredVisibility(const UE::Editor::FSubobjectEditorContext& Context, EVisibility& VisibilityOut) const
+{
+	for (const TSharedPtr<ISCSEditorUICustomization>& Customization : Customizations)
+	{
+		if (Customization->OverrideDesiredVisibility(Context, VisibilityOut))
+		{
+			return true;
+		}
+	}
+	return ISCSEditorUICustomization::OverrideDesiredVisibility(Context, VisibilityOut);
+}
+
 EChildActorComponentTreeViewVisualizationMode FActorDetailsSCSEditorUICustomization::GetChildActorVisualizationMode() const
 {
 	for (const TSharedPtr<ISCSEditorUICustomization>& Customization : Customizations)
@@ -120,7 +132,7 @@ bool FActorDetailsSCSEditorUICustomization::FilterSubobjectData(TArray<FSubobjec
 	bool bChanged = false;
 	for (const TSharedPtr<ISCSEditorUICustomization>& Customization : Customizations)
 	{
-		bChanged = bChanged || Customization->FilterSubobjectData(SubobjectData);
+		bChanged = Customization->FilterSubobjectData(SubobjectData) || bChanged;
 	}
 	return bChanged;
 }
diff --git a/Engine/Source/Editor/LevelEditor/Private/ActorDetailsSCSEditorUICustomization.h b/Engine/Source/Editor/LevelEditor/Private/ActorDetailsSCSEditorUICustomization.h
--- a/Engine/Source/Editor/LevelEditor/Private/ActorDetailsSCSEditorUICustomization.h
+++ b/Engine/Source/Editor/LevelEditor/Private/ActorDetailsSCSEditorUICustomization.h
@@ -23,6 +23,7 @@ public:
 
 	virtual const FSlateBrush* GetIconBrush(const FSubobjectData&) const override;
 	virtual TSharedPtr<SWidget> GetControlsWidget(const FSubobjectData&) const override;
+	virtual bool OverrideDesiredVisibility(const UE::Editor::FSubobjectEditorContext& Context, EVisibility& VisibilityOut) const;
 
 	virtual EChildActorComponentTreeViewVisualizationMode GetChildActorVisualizationMode() const override;
 
diff --git a/Engine/Source/Editor/LevelEditor/Private/SActorDetails.cpp b/Engine/Source/Editor/LevelEditor/Private/SActorDetails.cpp
--- a/Engine/Source/Editor/LevelEditor/Private/SActorDetails.cpp
+++ b/Engine/Source/Editor/LevelEditor/Private/SActorDetails.cpp
@@ -575,12 +575,16 @@ void SActorDetails::RefreshSubobjectTreeElements(TArrayView<const FSubobjectEdit
 						const FSubobjectDataSubsystemContextData_TedsRow* TedsRowAdapter = Data->GetContextData().GetPtr<FSubobjectDataSubsystemContextData_TedsRow>();
 						if (TedsRowAdapter)
 						{
-							UE::Editor::DataStorage::RowHandle ComponentRowHandle = TedsRowAdapter->RowHandle;
-							if (const FTedsTypedElementColumn* TypedElementHandleColumn = CoreProvider->GetColumn<FTedsTypedElementColumn>(ComponentRowHandle))
+							RowHandle ComponentRowHandle = TedsRowAdapter->GetRow();
+
+							FHierarchyHandle ReferenceObjectHierarchyName = FSubobjectDataSubsystemContextData_TedsRow::GetReferencedDataHierarchy(*CoreProvider);
+							CoreProvider->WalkDepthFirst(ReferenceObjectHierarchyName, ComponentRowHandle, [&Components](const ICoreProvider& Context, RowHandle Owner, RowHandle Target)
 							{
-								FTypedElementHandle ComponentElementHandle = TypedElementHandleColumn->Handle;
-								Components.Add(TypedElementHandleColumn->Handle);
-							}
+								if (const Compatibility::FTypedElementColumn* TypedElementColumn = Context.GetColumn<UE::Editor::DataStorage::Compatibility::FTypedElementColumn>(Target))
+								{
+									Components.Add(TypedElementColumn->Handle);
+								}
+							});
 						}
 					}
 				}
@@ -791,7 +795,7 @@ void SActorDetails::OnSubobjectEditorTreeViewSelectionChanged(const TArray<FSubo
 		// Update the underlying details view
 		RefreshSubobjectTreeElements(SelectedNodes, /*bForceRefresh*/false, DetailsView->IsLocked());
 	}
-	else if (UObject* ObjectContext = GetObjectContext())
+	else if (!TopLevelElements.IsEmpty())
 	{
 		RefreshSubobjectTreeElements(SelectedNodes, /*bForceRefresh*/false, DetailsView->IsLocked());
 	}
@@ -928,8 +932,12 @@ EVisibility SActorDetails::GetComponentEditorVisibility() const
 	// see if we need to hide the editor due to the current object display 
 	const bool bHideEditorFromDetailsView = (DisplayManager.IsValid() &&
 							                 DisplayManager->ShouldHideComponentEditor() );
+	if (bHideEditorFromDetailsView)
+	{
+		return EVisibility::Collapsed;
+	}
 
-	return GetObjectContext() && !bHideEditorFromDetailsView  ? EVisibility::Visible : EVisibility::Collapsed;
+	return SubobjectEditor.IsValid() ? SubobjectEditor->GetDesiredVisibility() : EVisibility::Collapsed;
 }
 
 EVisibility SActorDetails::GetComponentEditorButtonsVisibility() const
diff --git a/Engine/Source/Editor/SubobjectDataInterface/Private/SubobjectDataAdapterBase.cpp b/Engine/Source/Editor/SubobjectDataInterface/Private/SubobjectDataAdapterBase.cpp
--- a/Engine/Source/Editor/SubobjectDataInterface/Private/SubobjectDataAdapterBase.cpp
+++ b/Engine/Source/Editor/SubobjectDataInterface/Private/SubobjectDataAdapterBase.cpp
@@ -1,11 +1,13 @@
-﻿// Copyright Epic Games, Inc. All Rights Reserved.
+// Copyright Epic Games, Inc. All Rights Reserved.
 
 #include "SubobjectDataAdapterBase.h"
 
 #include "DataStorage/Features.h"
+#include "DataStorage/ScopedRowHandle.h"
 #include "Elements/Columns/TypedElementCompatibilityColumns.h"
 #include "Elements/Interfaces/TypedElementDataStorageInterface.h"
 #include "SubobjectDataHandle.h"
+#include "SubobjectDataTedsFactory.h"
 
 #define LOCTEXT_NAMESPACE "FSubobjectDataBehaviorBase"
 
@@ -18,54 +20,46 @@ FSubobjectDataSubsystemContextData_SingleUObjectContextObject::FSubobjectDataSub
 {
 }
 
-void UE::SubobjectData::QueryContextDataForObjects(const TInstancedStruct<FSubobjectDataSubsystemContextDataBase>& ContextData, TArray<const UObject*>& OutObjects, const bool bEvenIfPendingKill)
+FSubobjectDataSubsystemContextData_TedsRow::FSubobjectDataSubsystemContextData_TedsRow()
+{
+}
+
+FSubobjectDataSubsystemContextData_TedsRow::FSubobjectDataSubsystemContextData_TedsRow(UE::Editor::DataStorage::RowHandle InRowHandle)
 {
-	if (const FSubobjectDataSubsystemContextData_SingleUObjectContextObject* SingleObjectData = ContextData.GetPtr<FSubobjectDataSubsystemContextData_SingleUObjectContextObject>())
+	using namespace UE::Editor::DataStorage;
+	if (InRowHandle == InvalidRowHandle)
 	{
-		if (const UObject* PinnedObject = SingleObjectData->Object.Get(bEvenIfPendingKill))
-		{
-			OutObjects.Add(PinnedObject);
-		}
+		return;
 	}
-	else if (const FSubobjectDataSubsystemContextData_TedsRow* TedsRowData = ContextData.GetPtr<FSubobjectDataSubsystemContextData_TedsRow>())
+	ICoreProvider* EditorDataStorage = GetMutableDataStorageFeature<ICoreProvider>(StorageFeatureName);
+	if (!ensure(EditorDataStorage))
 	{
-		using namespace UE::Editor::DataStorage;
-
-		if (ICoreProvider* EditorDataStorage = GetMutableDataStorageFeature<ICoreProvider>(StorageFeatureName))
-		{
-			if (FTypedElementUObjectColumn* ObjectColumn = EditorDataStorage->GetColumn<FTypedElementUObjectColumn>(TedsRowData->RowHandle))
-			{
-				if (const UObject* ObjectValue = ObjectColumn->Object.Get(bEvenIfPendingKill))
-				{
-					OutObjects.Add(ObjectValue);
-				}
-			}
-		}
+		return;
 	}
+	Row = MakeShared<FScopedRow>(EditorDataStorage, InRowHandle);
 }
 
-FSubobjectDataSubsystemContextData_TedsRow::FSubobjectDataSubsystemContextData_TedsRow()
-{
-}
-
-FSubobjectDataSubsystemContextData_TedsRow::FSubobjectDataSubsystemContextData_TedsRow(UE::Editor::DataStorage::RowHandle InRowHandle)
+UE::Editor::DataStorage::TableHandle FSubobjectDataSubsystemContextData_TedsRow::GetSubobjectItemTable(
+	const UE::Editor::DataStorage::ICoreProvider& DataStorage)
 {
-	RowHandle = InRowHandle;
+	const USubobjectDataTedsFactory* Factory = DataStorage.FindFactory<USubobjectDataTedsFactory>();
+	return Factory->GetSubobjectDataTableHandle();
 }
 
-FSubobjectDataSubsystemContextData_TedsRow::FSubobjectDataSubsystemContextData_TedsRow(const FTedsRowHandle& InRowHandle)
-	: RowHandle(InRowHandle)
+UE::Editor::DataStorage::FHierarchyHandle FSubobjectDataSubsystemContextData_TedsRow::GetReferencedDataHierarchy(
+	const UE::Editor::DataStorage::ICoreProvider& DataStorage)
 {
+	return DataStorage.FindHierarchyByName(USubobjectDataTedsFactory::GetReferencedDataHierarchy());
 }
 
 FSubobjectDataSubsystemAdapterBase::~FSubobjectDataSubsystemAdapterBase() = default;
 
-bool FSubobjectDataSubsystemAdapterBase::CanIntercept_GatherSubobjectData(const TInstancedStruct<FSubobjectDataSubsystemContextDataBase>& ContextData) const
+bool FSubobjectDataSubsystemAdapterBase::CanIntercept_GatherSubobjectData(const UE::Editor::FSubobjectEditorContext& Context) const
 {
 	return false;
 }
 
-void FSubobjectDataSubsystemAdapterBase::GatherSubobjectData(const TInstancedStruct<FSubobjectDataSubsystemContextDataBase>& ContextData, TArray<FSubobjectDataHandle>& OutArray, TFunctionRef<FSubobjectDataHandle()> CreateSubobjectData) const
+void FSubobjectDataSubsystemAdapterBase::GatherSubobjectData(const UE::Editor::FSubobjectEditorContext& Context, TArray<FSubobjectDataHandle>& OutArray, TFunctionRef<FSubobjectDataHandle()> CreateSubobjectData) const
 {
 }
 
@@ -402,4 +396,4 @@ FText FSubobjectDataAdapterBase::GetActorDisplayText(const FSubobjectData& Data)
 	return FText();
 }
 
-#undef LOCTEXT_NAMESPACE
\ No newline at end of file
+#undef LOCTEXT_NAMESPACE
diff --git a/Engine/Source/Editor/SubobjectDataInterface/Private/SubobjectDataSubsystem.cpp b/Engine/Source/Editor/SubobjectDataInterface/Private/SubobjectDataSubsystem.cpp
--- a/Engine/Source/Editor/SubobjectDataInterface/Private/SubobjectDataSubsystem.cpp
+++ b/Engine/Source/Editor/SubobjectDataInterface/Private/SubobjectDataSubsystem.cpp
@@ -271,38 +271,41 @@ FSubobjectDataHandle USubobjectDataSubsystem::FindAttachParentForInheritedCompon
 	return OutHandle;
 }
 
-void USubobjectDataSubsystem::GatherSubobjectData(UObject* Context, TArray<FSubobjectDataHandle>& OutArray)
+void USubobjectDataSubsystem::GatherSubobjectData(UObject* ObjectContext, TArray<FSubobjectDataHandle>& OutArray)
 {
-	OutArray.Reset();
-	
-	FSubobjectDataSubsystemContextData_SingleUObjectContextObject ContextDataAdapter_OnStack(Context);
-	TConstStructView<FSubobjectDataSubsystemContextDataBase> StructView(ContextDataAdapter_OnStack);
-	TInstancedStruct<FSubobjectDataSubsystemContextDataBase> ContextDataAdapter(StructView);
+	UE::Editor::FSubobjectEditorContext Context;
+	Context.ObjectContext = ObjectContext;
+	GatherSubobjectData(Context, OutArray);
+}
+
+void USubobjectDataSubsystem::GatherSubobjectData(const UE::Editor::FSubobjectEditorContext& Context, TArray<FSubobjectDataHandle>& OutArray)
+{
+	OutArray.Reset();	
 
 	for (const auto& Entry : SubsystemAdapters)
 	{
-		if (Entry.Adapter->CanIntercept_GatherSubobjectData(ContextDataAdapter))
+		if (Entry.Adapter->CanIntercept_GatherSubobjectData(Context))
 		{
 			return Entry.Adapter->GatherSubobjectData(
-				ContextDataAdapter,
+				Context,
 				OutArray,
 				[this]() { return this->CreateSubobjectData(); });
 		}
 	}
 	
-	if (!Context)
+	if (!Context.ObjectContext)
 	{
 		UE_LOG(LogSubobjectSubsystem, Warning, TEXT("Could not gather subobject data, there was no context given!"));
 		return;
 	}
 
-	const AActor* ActorContext = Cast<AActor>(Context);
+	const AActor* ActorContext = Cast<AActor>(Context.ObjectContext);
 	if(!ActorContext)
 	{
 		return;
 	}
 
-	FSubobjectDataHandle RootActorHandle = CreateSubobjectData(Context);
+	FSubobjectDataHandle RootActorHandle = CreateSubobjectData(Context.ObjectContext);
 	OutArray.Add(RootActorHandle);
 	const FSubobjectData* RootActorDataPtr = RootActorHandle.GetData();
 	check(RootActorDataPtr);
diff --git a/Engine/Source/Editor/SubobjectDataInterface/Private/SubobjectDataTedsFactory.cpp b/Engine/Source/Editor/SubobjectDataInterface/Private/SubobjectDataTedsFactory.cpp
new file mode 100644
--- /dev/null
+++ b/Engine/Source/Editor/SubobjectDataInterface/Private/SubobjectDataTedsFactory.cpp
@@ -0,0 +1,35 @@
+// Copyright Epic Games, Inc. All Rights Reserved.
+
+#include "SubobjectDataTedsFactory.h"
+
+#include "SubobjectDataAdapterBase.h"
+#include "Elements/Interfaces/TypedElementDataStorageInterface.h"
+
+namespace SubobjectDataTedsFactory::Private
+{
+	FName ReferenceDataHierarchyName = TEXT("SubobjectDataReferencedDataHierarchy");
+}
+
+void USubobjectDataTedsFactory::RegisterHierarchies(UE::Editor::DataStorage::ICoreProvider& DataStorage)
+{
+	UE::Editor::DataStorage::FHierarchyRegistrationParams Params;
+	Params.Name = GetReferencedDataHierarchy();
+	DataStorage.RegisterHierarchy(Params);
+	Super::RegisterHierarchies(DataStorage);
+}
+
+void USubobjectDataTedsFactory::RegisterTables(UE::Editor::DataStorage::ICoreProvider& DataStorage)
+{
+	SubobjectDataTableHandle = DataStorage.RegisterTable<FSubobjectDataTag>(TEXT("SubobjectDataTable"));
+	Super::RegisterTables(DataStorage);
+}
+
+UE::Editor::DataStorage::TableHandle USubobjectDataTedsFactory::GetSubobjectDataTableHandle() const
+{
+	return SubobjectDataTableHandle;
+}
+
+FName USubobjectDataTedsFactory::GetReferencedDataHierarchy()
+{
+	return SubobjectDataTedsFactory::Private::ReferenceDataHierarchyName;
+}
diff --git a/Engine/Source/Editor/SubobjectDataInterface/Private/SubobjectDataTedsFactory.h b/Engine/Source/Editor/SubobjectDataInterface/Private/SubobjectDataTedsFactory.h
new file mode 100644
--- /dev/null
+++ b/Engine/Source/Editor/SubobjectDataInterface/Private/SubobjectDataTedsFactory.h
@@ -0,0 +1,27 @@
+// Copyright Epic Games, Inc. All Rights Reserved.
+
+#pragma once
+#include "DataStorage/Handles.h"
+#include "Elements/Interfaces/TypedElementDataStorageFactory.h"
+
+#include "SubobjectDataTedsFactory.generated.h"
+
+UCLASS()
+class USubobjectDataTedsFactory final : public UEditorDataStorageFactory
+{
+	GENERATED_BODY()
+
+public:
+	~USubobjectDataTedsFactory() override = default;
+
+	void RegisterHierarchies(UE::Editor::DataStorage::ICoreProvider& DataStorage) override;
+	void RegisterTables(UE::Editor::DataStorage::ICoreProvider& DataStorage) override;
+
+	UE::Editor::DataStorage::TableHandle GetSubobjectDataTableHandle() const;
+
+	static FName GetReferencedDataHierarchy();
+	
+
+private:
+	UE::Editor::DataStorage::TableHandle SubobjectDataTableHandle = UE::Editor::DataStorage::InvalidTableHandle;
+};
\ No newline at end of file
diff --git a/Engine/Source/Editor/SubobjectDataInterface/Public/SubobjectDataAdapterBase.h b/Engine/Source/Editor/SubobjectDataInterface/Public/SubobjectDataAdapterBase.h
--- a/Engine/Source/Editor/SubobjectDataInterface/Public/SubobjectDataAdapterBase.h
+++ b/Engine/Source/Editor/SubobjectDataInterface/Public/SubobjectDataAdapterBase.h
@@ -1,10 +1,10 @@
-﻿// Copyright Epic Games, Inc. All Rights Reserved.
+// Copyright Epic Games, Inc. All Rights Reserved.
 
 #pragma once
 
 #include "CoreMinimal.h"
-// #include "SubobjectDataHandle.h"
 #include "DataStorage/Handles.h"
+#include "DataStorage/ScopedRowHandle.h"
 #include "StructUtils/InstancedStruct.h"
 
 #include "SubobjectDataAdapterBase.generated.h"
@@ -15,9 +15,13 @@ class UBlueprint;
 class UObject;
 struct FSubobjectData;
 struct FSubobjectDataHandle;
-
 struct FTedsRowHandle;
 
+namespace UE::Editor
+{
+	struct FSubobjectEditorContext;
+}
+
 #define UE_API SUBOBJECTDATAINTERFACE_API
 
 struct
@@ -43,6 +47,12 @@ FSubobjectDataSubsystemContextDataBase
 	GENERATED_BODY()
 };
 
+USTRUCT()
+struct FSubobjectDataTag : public FEditorDataStorageTag
+{
+	GENERATED_BODY()
+};
+
 // An implementation of the ContextData that supports containing a single UObject
 // This is to provide interoperability with the existing ContextObject used in APIs of the SubobjectDataSubsystem
 USTRUCT()
@@ -66,16 +76,24 @@ FSubobjectDataSubsystemContextData_TedsRow : public FSubobjectDataSubsystemConte
 
 	UE_API FSubobjectDataSubsystemContextData_TedsRow();
 	UE_API explicit FSubobjectDataSubsystemContextData_TedsRow(UE::Editor::DataStorage::RowHandle InRowHandle);
-	UE_API explicit FSubobjectDataSubsystemContextData_TedsRow(const FTedsRowHandle& InRowHandle);
 
-	UPROPERTY()
-	FTedsRowHandle RowHandle;
+	UE_API UE::Editor::DataStorage::RowHandle GetRow() const;
+
+	UE_API static UE::Editor::DataStorage::TableHandle GetSubobjectItemTable(const UE::Editor::DataStorage::ICoreProvider& DataStorage);
+	
+	// Can be used along with TEDS API to walk any data rows that are referenced by this SubobjectData
+	// For example, a SubobjectData could represent multiple UActorComponent instances, they would be represented
+	// as children in the ReferencedDataHierarchy
+	UE_API static UE::Editor::DataStorage::FHierarchyHandle GetReferencedDataHierarchy(const UE::Editor::DataStorage::ICoreProvider& DataStorage);
+private:
+
+	TSharedPtr<UE::Editor::DataStorage::FScopedRow> Row;
 };
 
-namespace UE::SubobjectData
+inline UE::Editor::DataStorage::RowHandle FSubobjectDataSubsystemContextData_TedsRow::GetRow() const
 {
-	UE_API void QueryContextDataForObjects(const TInstancedStruct<FSubobjectDataSubsystemContextDataBase>& ContextData, TArray<const UObject*>& OutObjects, bool bEvenIfPendingKill = false);
-} // namespace UE::SubobjectData
+	return Row ? Row->GetRow() : UE::Editor::DataStorage::InvalidRowHandle;
+}
 
 /**
  * The base of an adapter/plugin that extends the SubobjectDataSubsystem by intercepting
@@ -89,8 +107,8 @@ class FSubobjectDataSubsystemAdapterBase
 {
 public:
 	UE_API virtual ~FSubobjectDataSubsystemAdapterBase();
-	UE_API virtual bool CanIntercept_GatherSubobjectData(const TInstancedStruct<FSubobjectDataSubsystemContextDataBase>& ContextData) const;
-	UE_API virtual void GatherSubobjectData(const TInstancedStruct<FSubobjectDataSubsystemContextDataBase>& ContextData, TArray<FSubobjectDataHandle>& OutArray, TFunctionRef<FSubobjectDataHandle()> CreateSubobjectData) const;
+	UE_API virtual bool CanIntercept_GatherSubobjectData(const UE::Editor::FSubobjectEditorContext& Context) const;
+	UE_API virtual void GatherSubobjectData(const UE::Editor::FSubobjectEditorContext& Context, TArray<FSubobjectDataHandle>& OutArray, TFunctionRef<FSubobjectDataHandle()> CreateSubobjectData) const;
 
 	// Call this implementation of FindSceneRootForSubobject if returning true
 	UE_API virtual bool CanIntercept_FindSceneRootForSubobject(const FSubobjectDataHandle& InHandle) const;
diff --git a/Engine/Source/Editor/SubobjectDataInterface/Public/SubobjectDataSubsystem.h b/Engine/Source/Editor/SubobjectDataInterface/Public/SubobjectDataSubsystem.h
--- a/Engine/Source/Editor/SubobjectDataInterface/Public/SubobjectDataSubsystem.h
+++ b/Engine/Source/Editor/SubobjectDataInterface/Public/SubobjectDataSubsystem.h
@@ -86,6 +86,26 @@ struct FReparentSubobjectParams
 	TObjectPtr<AActor> ActorPreviewContext = nullptr;
 };
 
+namespace UE::Editor
+{
+	//~ It would be nice if this were a nested class inside USubobjectDataSubsystem, but
+	//~  we need it in SubobjectDataAdapterBase.h, we can't forward declare nested classes,
+	//~  and it's a pain to forward declare everything from that file here instead.
+	/**
+	 * Struct holding information about the context of a subobject editor
+	 */
+	struct FSubobjectEditorContext
+	{
+		// If editor operates on a single object, the object to use. I.e. the object
+		//  typically returned by GetObjectContext inside a details panel. This might
+		//  eventually be phased out in favor of always using SelectionSetName or
+		//  some other way to identify the target objects (since there may be multiple)
+		UObject* ObjectContext = nullptr;
+		// If editor operates on a selection set, its identifier
+		FName SelectionSetName;
+	};
+}
+
 /**
 * The Subobject Data Subsystem will produce the reflected subobject data
 * based on a given root object. A root object can be anything, an actor 
@@ -114,10 +134,21 @@ public:
 	* Gather all subobjects that the given UObject context has. Populates an array of 
 	* handles that will have the given context and all it's subobjects.
 	* 
-	* @param Context		Object to gather subobjects for
+	* @param ObjectContext	Object to gather subobjects for
 	* @param OutArray		Array to populate (will be emptied first)
 	*/
-	UE_API void GatherSubobjectData(UObject* Context, TArray<FSubobjectDataHandle>& OutArray);
+	UE_API void GatherSubobjectData(UObject* ObjectContext, TArray<FSubobjectDataHandle>& OutArray);
+
+	/**
+	* Gather all subobjects that the given context has. Populates an array of
+	* handles that will have the relevant subobjects (taken from selection, or from
+	* an object, as appropriate).
+	*
+	* @param OutArray		Array to populate (will be emptied first)
+	*/
+	UE_EXPERIMENTAL(5.8, "Use GatherSubobjectData(UObject*, ...) for now.  This API may change.")
+	UE_API void GatherSubobjectData(const UE::Editor::FSubobjectEditorContext& Context, TArray<FSubobjectDataHandle>& OutArray);
+
 
 	////////////////////////////////////////////
 	// System events
diff --git a/Engine/Source/Editor/SubobjectEditor/Private/SSubobjectEditor.cpp b/Engine/Source/Editor/SubobjectEditor/Private/SSubobjectEditor.cpp
--- a/Engine/Source/Editor/SubobjectEditor/Private/SSubobjectEditor.cpp
+++ b/Engine/Source/Editor/SubobjectEditor/Private/SSubobjectEditor.cpp
@@ -1,6 +1,7 @@
 // Copyright Epic Games, Inc. All Rights Reserved.
 
 #include "SSubobjectEditor.h"
+
 #include "Editor.h"
 #include "ClassIconFinder.h"
 #include "ScopedTransaction.h"
@@ -60,6 +61,8 @@
 #include "K2Node_Variable.h"		// Nodes for FSubobjectRowDragDropOp
 #include "K2Node_VariableGet.h"
 #include "AssetSelection.h"
+#include "DataStorage/Features.h"
+#include "Elements/Columns/TypedElementCompatibilityColumns.h"
 #include "Settings/BlueprintEditorProjectSettings.h"
 #include "Misc/FeedbackContext.h"
 
@@ -2338,6 +2341,20 @@ FSubobjectEditorTreeNodePtrType SSubobjectEditor::GetSceneRootNode() const
 	return nullptr;
 }
 
+EVisibility SSubobjectEditor::GetDesiredVisibility() const
+{
+	EVisibility CustomizationVisibility = EVisibility::Collapsed;
+	UE::Editor::FSubobjectEditorContext Context;
+	Context.ObjectContext = GetObjectContext();
+	Context.SelectionSetName = SelectionSetName;
+	if (UICustomization && UICustomization->OverrideDesiredVisibility(Context, CustomizationVisibility))
+	{
+		return CustomizationVisibility;
+	}
+
+	return Context.ObjectContext ? EVisibility::Visible : EVisibility::Collapsed;
+}
+
 void SSubobjectEditor::UpdateSelectionFromNodes(const TArray<FSubobjectEditorTreeNodePtrType>& SelectedNodes)
 {
 	bUpdatingSelection = true;
@@ -2524,7 +2541,7 @@ void SSubobjectEditor::UpdateTree(bool bRegenerateTreeNodes /* = true */)
 		}
 		RootNodes.Empty();
 
-		if (UObject* Context = GetObjectContext())
+		if (GetDesiredVisibility() == EVisibility::Visible)
 		{
 			ensureMsgf(FModuleManager::Get().LoadModule("SubobjectDataInterface"), TEXT("The Subobject Data Interface module is required."));
 
@@ -2533,7 +2550,10 @@ void SSubobjectEditor::UpdateTree(bool bRegenerateTreeNodes /* = true */)
 
 			TArray<FSubobjectDataHandle> SubobjectData;
 
-			DataSubsystem->GatherSubobjectData(Context, SubobjectData);
+			UE::Editor::FSubobjectEditorContext Params;
+			Params.ObjectContext = GetObjectContext();
+			Params.SelectionSetName = SelectionSetName;
+			DataSubsystem->GatherSubobjectData(Params, SubobjectData);
 
 			if (UICustomization)
 			{
@@ -2696,6 +2716,11 @@ FReply SSubobjectEditor::OnKeyDown(const FGeometry& MyGeometry, const FKeyEvent&
 	return FReply::Unhandled();
 }
 
+void SSubobjectEditor::SetSelectionSetName(FName SelectionSetNameIn)
+{
+	SelectionSetName = SelectionSetNameIn;
+}
+
 FReply SSubobjectEditor::TryHandleAssetDragDropOperation(const FDragDropEvent& DragDropEvent)
 {
 	TSharedPtr<FDragDropOperation> Operation = DragDropEvent.GetOperation();
@@ -3347,7 +3372,7 @@ void SSubobjectEditor::GetSelectedItemsForContextMenu(TArray<FComponentEventCons
 		NewItem.VariableName = Data->GetVariableName();
 		const UObject* ContextObject = GetObjectContext();
 		const AActor* ContextAsActor = Cast<AActor>(ContextObject);
-		if (!ContextObject->HasAnyFlags(RF_ClassDefaultObject) && ContextAsActor)
+		if (ContextObject && !ContextObject->HasAnyFlags(RF_ClassDefaultObject) && ContextAsActor)
 		{
 			NewItem.Component = const_cast<UActorComponent*>(Data->FindComponentInstanceInActor(ContextAsActor));
 		}
@@ -3357,13 +3382,36 @@ void SSubobjectEditor::GetSelectedItemsForContextMenu(TArray<FComponentEventCons
 		}
 		else if (Data->GetContextData().IsValid())
 		{
-			TArray<const UObject*> ContextDataObjects;
-			UE::SubobjectData::QueryContextDataForObjects(Data->GetContextData(), ContextDataObjects);
+			const TInstancedStruct<FSubobjectDataSubsystemContextDataBase>& ContextData = Data->GetContextData();
+			const auto* AsTedsRow = ContextData.GetPtr<FSubobjectDataSubsystemContextData_TedsRow>();
+			const auto* AsSingleUObject = ContextData.GetPtr<FSubobjectDataSubsystemContextData_SingleUObjectContextObject>();
 
-			for (const UObject* ContextDataObject : ContextDataObjects)
+			if (AsTedsRow)
 			{
-				FComponentEventConstructionData& ContextItem = OutSelectedItems.AddDefaulted_GetRef();
-				ContextItem.Component = const_cast<UObject*>(ContextDataObject);
+				using namespace UE::Editor::DataStorage;
+				ICoreProvider* CoreProvider = UE::Editor::DataStorage::GetMutableDataStorageFeature<ICoreProvider>(StorageFeatureName);
+				check(CoreProvider);
+				
+				FHierarchyHandle ReferenceObjectHierarchyName = FSubobjectDataSubsystemContextData_TedsRow::GetReferencedDataHierarchy(*CoreProvider);
+				RowHandle ComponentRowHandle = AsTedsRow->GetRow();
+				CoreProvider->WalkDepthFirst(ReferenceObjectHierarchyName, ComponentRowHandle, [&OutSelectedItems](const ICoreProvider& Context, RowHandle Owner, RowHandle Target)
+				{
+					const FTypedElementUObjectColumn* UObjectColumn = Context.GetColumn<FTypedElementUObjectColumn>(Target);
+					if (UObjectColumn)
+					{
+						UObject* Object = UObjectColumn->Object.Get();
+						FComponentEventConstructionData& ContextItem = OutSelectedItems.AddDefaulted_GetRef();
+						ContextItem.Component = Object;
+					}
+				});
+			}
+			else if (AsSingleUObject)
+			{
+				if (TStrongObjectPtr<UObject> Object = AsSingleUObject->Object.Pin())
+				{
+					FComponentEventConstructionData& ContextItem = OutSelectedItems.AddDefaulted_GetRef();
+					ContextItem.Component = Object.Get();
+				}
 			}
 
 			// Skip adding the partially constructed NewItem since we handled that in our loop above
diff --git a/Engine/Source/Editor/SubobjectEditor/Public/SSubobjectEditor.h b/Engine/Source/Editor/SubobjectEditor/Public/SSubobjectEditor.h
--- a/Engine/Source/Editor/SubobjectEditor/Public/SSubobjectEditor.h
+++ b/Engine/Source/Editor/SubobjectEditor/Public/SSubobjectEditor.h
@@ -468,6 +468,12 @@ protected:
 	UE_API virtual FReply OnKeyDown(const FGeometry& MyGeometry, const FKeyEvent& InKeyEvent) override;
 	
 public:
+	/** 
+	 * Sets the identifier representing this subobject editor's selection in TEDS.
+	 *  If not specified, will use the level editor selection. 
+	 */
+	UE_API void SetSelectionSetName(FName SelectionSetNameIn);
+
 	/** Select the given tree node */
 	UE_API void SelectNode(FSubobjectEditorTreeNodePtrType InNodeToSelect, bool bIsCntrlDown);
 
@@ -502,6 +508,14 @@ public:
 	UE_API TArray<FSubobjectDataHandle> GetSelectedHandles() const;
 
 	UE_API virtual FSubobjectEditorTreeNodePtrType GetSceneRootNode() const;
+
+	/** 
+	 * Returns whether the editor wants to be visible, which can be optionally used by the containing
+	 *  widget. If not overriden by a customization, this will be Collapsed or Visible depending on 
+	 *  whether ObjectContext was null or not, respectively, but a customization can drive this by
+	 *  some other state.
+	 */
+	UE_API EVisibility GetDesiredVisibility() const;
 	
 	/** Try to handle a drag-drop operation */
 	UE_API FReply TryHandleAssetDragDropOperation(const FDragDropEvent& DragDropEvent);
@@ -769,6 +783,9 @@ protected:
 		/** Populate an array of subobject handles based on an array of the slate node ptr */
 		static void PopulateHandlesArray(const TArray<FSubobjectEditorTreeNodePtrType>& SlateNodePtrs, TArray<FSubobjectDataHandle>& OutHandles);
 	};
+
+private:
+	FName SelectionSetName;
 };
 
 #undef UE_API
diff --git a/Engine/Source/Editor/SubobjectEditor/SubobjectEditor.Build.cs b/Engine/Source/Editor/SubobjectEditor/SubobjectEditor.Build.cs
--- a/Engine/Source/Editor/SubobjectEditor/SubobjectEditor.Build.cs
+++ b/Engine/Source/Editor/SubobjectEditor/SubobjectEditor.Build.cs
@@ -33,6 +33,7 @@ public class SubobjectEditor : ModuleRules
 				"GraphEditor",
 				"ToolMenus",
 				"ToolWidgets",
+				"TypedElementFramework"
 			}
 		);
 		
