Enable optional overlay widget on the tree view in SSubobjectBlueprintEditor, similar to SSubobjectInstanceEditor

diff --git a/Engine/Source/Editor/SubobjectEditor/Private/SSubobjectBlueprintEditor.cpp b/Engine/Source/Editor/SubobjectEditor/Private/SSubobjectBlueprintEditor.cpp
--- a/Engine/Source/Editor/SubobjectEditor/Private/SSubobjectBlueprintEditor.cpp
+++ b/Engine/Source/Editor/SubobjectEditor/Private/SSubobjectBlueprintEditor.cpp
@@ -144,13 +144,37 @@ void SSubobjectBlueprintEditor::Construct(const FArguments& InArgs)
 
 	+ SVerticalBox::Slot()
 	[
-		SNew(SBorder)
-		.BorderImage(FAppStyle::Get().GetBrush("SCSEditor.Background"))
-		.Padding(4.f)
-		.AddMetaData<FTagMetaData>(FTagMetaData(TEXT("ComponentsPanel")))
-		.Visibility(this, &SSubobjectBlueprintEditor::GetComponentsTreeVisibility)
+		SNew(SOverlay)
+		+ SOverlay::Slot()
 		[
-			TreeWidget.ToSharedRef()
+			SNew(SBorder)
+			.BorderImage(FAppStyle::Get().GetBrush("SCSEditor.Background"))
+			.Padding(4.f)
+			.AddMetaData<FTagMetaData>(FTagMetaData(TEXT("ComponentsPanel")))
+			.Visibility(this, &SSubobjectBlueprintEditor::GetComponentsTreeVisibility)
+			[
+				TreeWidget.ToSharedRef()
+			]
+		]
+		+ SOverlay::Slot()
+		[
+			SAssignNew(TreeViewOverlay, SBox)
+			.Visibility_Lambda([this]()
+			{
+				const bool bTreeIsVisible = GetComponentsTreeVisibility() == EVisibility::Visible;
+				UObject* ObjectContextPtr = GetObjectContext();
+				const bool bSupportsOverlay = UICustomization && UICustomization->SupportsTreeViewOverlay(MakeArrayView(&ObjectContextPtr, 1));
+
+				if (bTreeIsVisible && bSupportsOverlay)
+				{
+					const bool bIsDragging = FSlateApplication::Get().IsDragDropping();
+					return bIsDragging ? EVisibility::Visible : EVisibility::HitTestInvisible;
+				}
+				return EVisibility::Collapsed;
+			})
+			[
+				UICustomization ? UICustomization->ConstructTreeViewOverlay(SharedThis(this)) : SNullWidget::NullWidget
+			]
 		]
 	];
 
@@ -184,6 +208,19 @@ void SSubobjectBlueprintEditor::Construct(const FArguments& InArgs)
 
 END_SLATE_FUNCTION_BUILD_OPTIMIZATION
 
+void SSubobjectBlueprintEditor::UpdateTree(bool bRegenerateTreeNodes)
+{
+	SSubobjectEditor::UpdateTree(bRegenerateTreeNodes);
+	if (TreeViewOverlay.IsValid())
+	{
+		UObject* ObjectContextPtr = GetObjectContext();
+		if (UICustomization && UICustomization->SupportsTreeViewOverlay(MakeArrayView(&ObjectContextPtr, 1)))
+		{
+			TreeViewOverlay->SetContent(UICustomization->ConstructTreeViewOverlay(SharedThis(this)));
+		}
+	}
+}
+
 void SSubobjectBlueprintEditor::OnDeleteNodes()
 {
 	const FScopedTransaction Transaction( LOCTEXT("RemoveComponents", "Remove Components") );
diff --git a/Engine/Source/Editor/SubobjectEditor/Public/SSubobjectBlueprintEditor.h b/Engine/Source/Editor/SubobjectEditor/Public/SSubobjectBlueprintEditor.h
--- a/Engine/Source/Editor/SubobjectEditor/Public/SSubobjectBlueprintEditor.h
+++ b/Engine/Source/Editor/SubobjectEditor/Public/SSubobjectBlueprintEditor.h
@@ -52,7 +52,8 @@ protected:
 	UE_API virtual void OnDuplicateComponent() override;
 	UE_API virtual bool CanPasteNodes() const override;
 	UE_API virtual void PasteNodes() override;
-	
+	UE_API virtual void UpdateTree(bool bRegenerateTreeNodes = true) override;
+
 	UE_API virtual void OnAttachToDropAction(FSubobjectEditorTreeNodePtrType DroppedOn, const TArray<FSubobjectEditorTreeNodePtrType>& DroppedNodePtrs);
     UE_API virtual void OnDetachFromDropAction(const TArray<FSubobjectEditorTreeNodePtrType>& DroppedNodePtrs);
     UE_API virtual void OnMakeNewRootDropAction(FSubobjectEditorTreeNodePtrType DroppedNodePtr);
@@ -121,6 +122,8 @@ protected:
 
 	/** Get the current preview actor for this blueprint editor. */
 	AActor* GetActorPreview() const { return PreviewActor.Get(nullptr); }
+
+	TSharedPtr<class SBox> TreeViewOverlay;
 };
 
 #undef UE_API
