Fix crash when opening Dataprep Asset

diff --git a/Engine/Source/Editor/SubobjectEditor/Private/SSubobjectInstanceEditor.cpp b/Engine/Source/Editor/SubobjectEditor/Private/SSubobjectInstanceEditor.cpp
--- a/Engine/Source/Editor/SubobjectEditor/Private/SSubobjectInstanceEditor.cpp
+++ b/Engine/Source/Editor/SubobjectEditor/Private/SSubobjectInstanceEditor.cpp
@@ -180,7 +180,7 @@ void SSubobjectInstanceEditor::Construct(const FArguments& InArgs)
 			{
 				const bool bTreeIsVisible = GetComponentsTreeVisibility() == EVisibility::Visible;
 				UObject* ObjectContextPtr = GetObjectContext();
-				const bool bSupportsOverlay = UICustomization->SupportsTreeViewOverlay(MakeArrayView(&ObjectContextPtr, 1));		
+				const bool bSupportsOverlay = UICustomization && UICustomization->SupportsTreeViewOverlay(MakeArrayView(&ObjectContextPtr, 1));
 				
 				if (bTreeIsVisible && bSupportsOverlay)
 				{
@@ -193,7 +193,7 @@ void SSubobjectInstanceEditor::Construct(const FArguments& InArgs)
 			})
 			.Content()
 			[
-				UICustomization->ConstructTreeViewOverlay(SharedThis(this))
+				UICustomization ? UICustomization->ConstructTreeViewOverlay(SharedThis(this)) : SNullWidget::NullWidget
 			]
 			
 		]
@@ -235,7 +235,7 @@ void SSubobjectInstanceEditor::UpdateTree(bool bRegenerateTreeNodes)
 	if (TreeViewOverlay.IsValid())
 	{
 		UObject* ObjectContextPtr = GetObjectContext();
-		if (UICustomization->SupportsTreeViewOverlay(MakeArrayView(&ObjectContextPtr, 1)))
+		if (UICustomization && UICustomization->SupportsTreeViewOverlay(MakeArrayView(&ObjectContextPtr, 1)))
 		{
 			TreeViewOverlay->SetContent(UICustomization->ConstructTreeViewOverlay(SharedThis(this)));
 		}
