Undo 46339718 and just manually purge UI state to avoid leaks in SSubobjectEditor. The SubobjectData API is exposed to script and we changed semantic. If we want to make it leak proof we must reimplement in terms of UObject*

#rnx

diff --git a/Engine/Source/Editor/SubobjectDataInterface/Private/SubobjectData.cpp b/Engine/Source/Editor/SubobjectDataInterface/Private/SubobjectData.cpp
--- a/Engine/Source/Editor/SubobjectDataInterface/Private/SubobjectData.cpp
+++ b/Engine/Source/Editor/SubobjectDataInterface/Private/SubobjectData.cpp
@@ -29,8 +29,10 @@ PRAGMA_DISABLE_DEPRECATION_WARNINGS
 	}
 PRAGMA_ENABLE_DEPRECATION_WARNINGS
 
+
 FSubobjectData::FSubobjectData()
 	: WeakObjectPtr(nullptr)
+	, ParentObjectHandle(FSubobjectDataHandle::InvalidHandle)
 	, SCSNodePtr(nullptr)
 {
 	// By Default the handle will be invalid, it will only
@@ -38,6 +40,7 @@ FSubobjectData::FSubobjectData()
 }
 
 FSubobjectData::FSubobjectData(UObject* ContextObject, const FSubobjectDataHandle& InParentHandle, const bool bInIsInheritedSCS)
+	: ParentObjectHandle(InParentHandle)
 {
 	Initialize(ContextObject, InParentHandle, bInIsInheritedSCS);
 }
@@ -94,7 +97,7 @@ void FSubobjectData::Initialize(TInstancedStruct<FSubobjectDataSubsystemContextD
 	}
 
 	ContextData = MoveTemp(InContextData);
-	ParentObjectHandle = ParentHandle.GetSharedDataPtr();
+	ParentObjectHandle = ParentHandle;
 }
 
 bool FSubobjectData::CanEdit() const
@@ -676,7 +679,7 @@ void FSubobjectData::SetupAttachment(FName SocketName, const FSubobjectDataHandl
 
 FSubobjectDataHandle FSubobjectData::GetRootSubobject() const
 {
-	FSubobjectDataHandle Current = ParentObjectHandle.Pin();
+	FSubobjectDataHandle Current = ParentObjectHandle;
 	while(Current.IsValid() && Current.GetSharedDataPtr()->HasParent())
 	{
 		Current = Current.GetSharedDataPtr()->GetParentHandle();
@@ -736,7 +739,7 @@ FSubobjectDataHandle FSubobjectData::FindChildByObject(UObject* ContextObject) c
 bool FSubobjectData::AddChildHandleOnly(const FSubobjectDataHandle& InHandle)
 {
 	// If we already have this child, then don't both with adding it
-	if (HasChild(InHandle) || !InHandle.IsValid() || InHandle == Handle.Pin())
+	if (HasChild(InHandle) || !InHandle.IsValid() || InHandle == Handle)
 	{
 		return false;
 	}
@@ -744,7 +747,7 @@ bool FSubobjectData::AddChildHandleOnly(const FSubobjectDataHandle& InHandle)
 	ChildrenHandles.Add(InHandle);
 	if(FSubobjectData* NewChildData = InHandle.GetData())
 	{
-		NewChildData->SetParentHandle(Handle.Pin());		
+		NewChildData->SetParentHandle(Handle);		
 	}
 
 	return true;
@@ -989,6 +992,12 @@ FText FSubobjectData::GetActorDisplayText() const
 	return FText::GetEmpty();
 }
 
+void FSubobjectData::DisposeUnsafe()
+{
+	Handle = FSubobjectDataHandle();
+	ParentObjectHandle = FSubobjectDataHandle();
+}
+
 const UObject* FSubobjectData::GetObject(bool bEvenIfPendingKill) const
 {
 	return WeakObjectPtr.Get(bEvenIfPendingKill);
@@ -1203,7 +1212,7 @@ bool FSubobjectData::IsInstancedInheritedComponent() const
 	{
 		return false;
 	}
-	FSubobjectDataHandle CurrentHandle = ParentObjectHandle.Pin();
+	FSubobjectDataHandle CurrentHandle = ParentObjectHandle;
 	FSubobjectData* CurrentData = CurrentHandle.GetData();
 			
 	while(CurrentHandle.IsValid() && CurrentData && !CurrentData->IsActor())
@@ -1219,7 +1228,7 @@ bool FSubobjectData::IsAttachedTo(const FSubobjectDataHandle& InHandle) const
 {
 	REDIRECT_TO_ADAPTER(IsAttachedTo, InHandle)
 
-	FSubobjectDataHandle TestParentHandle = ParentObjectHandle.Pin();
+	FSubobjectDataHandle TestParentHandle = ParentObjectHandle;
 
 	while (TestParentHandle.IsValid())
 	{
diff --git a/Engine/Source/Editor/SubobjectDataInterface/Public/SubobjectData.h b/Engine/Source/Editor/SubobjectDataInterface/Public/SubobjectData.h
--- a/Engine/Source/Editor/SubobjectDataInterface/Public/SubobjectData.h
+++ b/Engine/Source/Editor/SubobjectDataInterface/Public/SubobjectData.h
@@ -183,9 +183,9 @@ public:
 	UE_API void SetSocketName(FName InNewName);
 	UE_API void SetupAttachment(FName SocketName, const FSubobjectDataHandle& AttachParentHandle = FSubobjectDataHandle::InvalidHandle);
 
-	FSubobjectDataHandle GetHandle() const { return Handle.Pin(); }
-	FSubobjectDataHandle GetParentHandle() const { return ParentObjectHandle.Pin(); }
-	FSubobjectData* GetParentData() const { return ParentObjectHandle.Pin().Get(); }
+	FSubobjectDataHandle GetHandle() const { return Handle; }
+	FSubobjectDataHandle GetParentHandle() const { return ParentObjectHandle; }
+	FSubobjectData* GetParentData() const { return ParentObjectHandle.GetData(); }
 	bool HasParent() const { return ParentObjectHandle.IsValid(); }
 	const TArray<FSubobjectDataHandle>& GetChildrenHandles() const { return ChildrenHandles; }
 	
@@ -231,6 +231,11 @@ public:
 
 	UE_API FText GetActorDisplayText() const;
 
+	/** 
+	 * This type has livliness problems due to cycles in the shared pointer graph. This function clears the pointers 
+	 *	involved in cycles, but live FSubobjectDatas will start to behave strangely once it is called
+	 */
+	UE_API void DisposeUnsafe();
 private:
 	/**
 	* Add the given subobject handle as a child of this.
@@ -248,18 +253,20 @@ private:
 	*/
 	UE_API bool RemoveChildHandleOnly(const FSubobjectDataHandle& InHandle);
 
-	void SetParentHandle(const FSubobjectDataHandle& NewParentHandle) { ParentObjectHandle = NewParentHandle.GetSharedDataPtr(); }
-	void ClearParentHandle() { ParentObjectHandle.Reset(); }
+	void SetParentHandle(const FSubobjectDataHandle& NewParentHandle) { ParentObjectHandle = NewParentHandle; }
+	void ClearParentHandle() { ParentObjectHandle = FSubobjectDataHandle::InvalidHandle; }
 	
 	// Weak ptr to the object instance represented by this node (e.g. component template)
 	UPROPERTY()
 	TWeakObjectPtr<UObject> WeakObjectPtr;
 
-	// Handle to this object
-	TWeakPtr<FSubobjectData> Handle;
+	// Handle that will have a globally unique ID when it is constructed (i.e. when this is constructed)
+	UPROPERTY()
+	FSubobjectDataHandle Handle;
 
 	// Handle to this object's parent. If this is invalid, then it is the root subobject.
-	TWeakPtr<FSubobjectData> ParentObjectHandle;
+	UPROPERTY()
+	FSubobjectDataHandle ParentObjectHandle;
 
 	// Set of child subobjects that use this subobject as a parent
 	UPROPERTY()
diff --git a/Engine/Source/Editor/SubobjectDataInterface/Public/SubobjectDataHandle.h b/Engine/Source/Editor/SubobjectDataInterface/Public/SubobjectDataHandle.h
--- a/Engine/Source/Editor/SubobjectDataInterface/Public/SubobjectDataHandle.h
+++ b/Engine/Source/Editor/SubobjectDataInterface/Public/SubobjectDataHandle.h
@@ -16,8 +16,10 @@
 struct FSubobjectData;
 
 /**
-* A subobject handle is a wrapper around a TSharedPtr<FSubobjectData> intended for use
-* with Blueprints.
+* A subobject handle is a globally unique identifier for subobjects
+* Upon construction, the handle will be invalid. It is the responsibility
+* of the owning FSubobjectData to set the DataPtr once the subobject
+* data has validated that it has a good context.
 */
 USTRUCT(BlueprintType)
 struct FSubobjectDataHandle
diff --git a/Engine/Source/Editor/SubobjectEditor/Private/SSubobjectEditor.cpp b/Engine/Source/Editor/SubobjectEditor/Private/SSubobjectEditor.cpp
--- a/Engine/Source/Editor/SubobjectEditor/Private/SSubobjectEditor.cpp
+++ b/Engine/Source/Editor/SubobjectEditor/Private/SSubobjectEditor.cpp
@@ -198,7 +198,7 @@ bool FSubobjectEditorTreeNode::IsNativeComponent() const
 bool FSubobjectEditorTreeNode::IsRootActorNode() const
 {
 	// If there is no parent node pointer then this has to be top of the tree
-	return !Parent.IsValid();
+	return ParentNodePtr == nullptr;
 }
 
 FName FSubobjectEditorTreeNode::GetVariableName() const
@@ -238,7 +238,7 @@ bool FSubobjectEditorTreeNode::IsChildSubtreeNode() const
 
 bool FSubobjectEditorTreeNode::IsAttachedTo(FSubobjectEditorTreeNodePtrType InNodePtr) const
 {
-	FSubobjectEditorTreeNodePtrType TestParentPtr = Parent.Pin();
+	FSubobjectEditorTreeNodePtrType TestParentPtr = ParentNodePtr;
 	while (TestParentPtr.IsValid())
 	{
 		if (TestParentPtr == InNodePtr)
@@ -246,7 +246,7 @@ bool FSubobjectEditorTreeNode::IsAttachedTo(FSubobjectEditorTreeNodePtrType InNo
 			return true;
 		}
 
-		TestParentPtr = TestParentPtr->Parent.Pin();
+		TestParentPtr = TestParentPtr->ParentNodePtr;
 	}
 
 	return false;
@@ -254,7 +254,7 @@ bool FSubobjectEditorTreeNode::IsAttachedTo(FSubobjectEditorTreeNodePtrType InNo
 
 bool FSubobjectEditorTreeNode::IsDirectlyAttachedTo(FSubobjectEditorTreeNodePtrType InNodePtr) const
 {
-	return Parent.Pin() == InNodePtr;
+	return ParentNodePtr == InNodePtr;
 }
 
 bool FSubobjectEditorTreeNode::CanDelete() const
@@ -286,13 +286,13 @@ void FSubobjectEditorTreeNode::AddChild(FSubobjectEditorTreeNodePtrType AttachTo
 	}
 	
 	Children.AddUnique(AttachToPtr);
-	AttachToPtr->Parent = AsShared();
+	AttachToPtr->ParentNodePtr = AsShared();
 }
 
 void FSubobjectEditorTreeNode::RemoveChild(FSubobjectEditorTreeNodePtrType InChildNodePtr)
 {
 	Children.Remove(InChildNodePtr);
-	InChildNodePtr->Parent.Reset();
+	InChildNodePtr->ParentNodePtr.Reset();
 }
 
 FSubobjectEditorTreeNodePtrType FSubobjectEditorTreeNode::FindChild(const FSubobjectDataHandle& InHandle)
@@ -347,15 +347,15 @@ void FSubobjectEditorTreeNode::SetCachedFilterState(bool bMatchesFilter, bool bU
 void FSubobjectEditorTreeNode::ApplyFilteredStateToParent()
 {
 	FSubobjectEditorTreeNode* Child = this;
-	while (Child->Parent.IsValid())
+	while (Child->ParentNodePtr.IsValid())
 	{
-		FSubobjectEditorTreeNode* CurrentParent = Child->Parent.Pin().Get();
+		FSubobjectEditorTreeNode* Parent = Child->ParentNodePtr.Get();
 
 		if (!IsFlaggedForFiltration())
 		{
-			if ((CurrentParent->FilterFlags & EFilteredState::ChildMatches) == 0)
+			if ((Parent->FilterFlags & EFilteredState::ChildMatches) == 0)
 			{
-				CurrentParent->FilterFlags |= EFilteredState::ChildMatches;
+				Parent->FilterFlags |= EFilteredState::ChildMatches;
 			}
 			else
 			{
@@ -364,10 +364,10 @@ void FSubobjectEditorTreeNode::ApplyFilteredStateToParent()
 			}
 		}
 		// have to see if this was the only child contributing to this flag
-		else if (CurrentParent->FilterFlags & EFilteredState::ChildMatches)
+		else if (Parent->FilterFlags & EFilteredState::ChildMatches)
 		{
-			CurrentParent->FilterFlags &= ~EFilteredState::ChildMatches;
-			for (const FSubobjectEditorTreeNodePtrType& Sibling : CurrentParent->Children)
+			Parent->FilterFlags &= ~EFilteredState::ChildMatches;
+			for (const FSubobjectEditorTreeNodePtrType& Sibling : Parent->Children)
 			{
 				if (Sibling.Get() == Child)
 				{
@@ -376,18 +376,18 @@ void FSubobjectEditorTreeNode::ApplyFilteredStateToParent()
 
 				if (Sibling->FilterFlags & EFilteredState::FilteredInMask)
 				{
-					CurrentParent->FilterFlags |= EFilteredState::ChildMatches;
+					Parent->FilterFlags |= EFilteredState::ChildMatches;
 					break;
 				}
 			}
 
-			if (CurrentParent->FilterFlags & EFilteredState::ChildMatches)
+			if (Parent->FilterFlags & EFilteredState::ChildMatches)
 			{
 				// another child added the flag back
 				break;
 			}
 		}
-		Child = CurrentParent;
+		Child = Parent;
 	}
 }
 
@@ -488,6 +488,25 @@ void FSubobjectEditorTreeNode::CloseOngoingCreateTransaction()
 	OngoingCreateTransaction.Reset();
 }
 
+void FSubobjectEditorTreeNode::ForEachNodeImpl(TFunctionRef<void(FSubobjectEditorTreeNode& Data)> Ftor)
+{
+	for(const FSubobjectEditorTreeNodePtrType& Child : Children)
+	{
+		Child->ForEachNodeImpl(Ftor);
+	}
+	
+	Ftor(*this);
+}
+
+void FSubobjectEditorTreeNode::DisposeUnsafe()
+{
+	if(FSubobjectData* Data = GetDataSource())
+	{
+		Data->DisposeUnsafe();
+	}
+	ParentNodePtr = FSubobjectEditorTreeNodePtrType();
+}
+
 //////////////////////////////////////////////////////////////////////////
 // FSubobjectRowDragDropOp - The drag-drop operation triggered when dragging a row in the components tree
 
@@ -2539,6 +2558,13 @@ void SSubobjectEditor::UpdateTree(bool bRegenerateTreeNodes /* = true */)
 		{
 			TreeWidget->ClearSelection();
 		}
+		// recursively purge nodes because their ownership is a mess - and the expected
+		// cascaded delete does not happen automatically:
+		ForEachNode([](FSubobjectEditorTreeNode& Tree) 
+			{
+				Tree.DisposeUnsafe();
+			}
+		);
 		RootNodes.Empty();
 
 		if (GetDesiredVisibility() == EVisibility::Visible)
@@ -3776,6 +3802,14 @@ void SSubobjectEditor::Utils::PopulateHandlesArray(const TArray<FSubobjectEditor
     }
 }
 
+void SSubobjectEditor::ForEachNode(TFunctionRef<void(FSubobjectEditorTreeNode& Data)> Ftor)
+{
+	for(const TSharedPtr<class FSubobjectEditorTreeNode>& Node : RootNodes)
+	{
+		Node->ForEachNodeImpl(Ftor);
+	}
+}
+
 void SSubobjectEditor::RefreshComponentTypesList()
 {
 	if (ComponentClassCombo.IsValid())
diff --git a/Engine/Source/Editor/SubobjectEditor/Public/SSubobjectEditor.h b/Engine/Source/Editor/SubobjectEditor/Public/SSubobjectEditor.h
--- a/Engine/Source/Editor/SubobjectEditor/Public/SSubobjectEditor.h
+++ b/Engine/Source/Editor/SubobjectEditor/Public/SSubobjectEditor.h
@@ -49,9 +49,7 @@ public:
 	
 	UE_API explicit FSubobjectEditorTreeNode(const FSubobjectDataHandle& DataSource, bool InbIsSeperator = false);
 
-	PRAGMA_DISABLE_DEPRECATION_WARNINGS // complains about ParentNodePtr otherwise
 	~FSubobjectEditorTreeNode() = default;
-	PRAGMA_ENABLE_DEPRECATION_WARNINGS
 
 	bool IsSeperator() const { return bIsSeperator; } 
 
@@ -68,7 +66,7 @@ public:
 	FSubobjectDataHandle GetDataHandle() const { return DataHandle; }
 
 	/** Get the slate parent of this single tree node */
-	FSubobjectEditorTreeNodePtrType GetParent() const { return Parent.Pin(); }
+	FSubobjectEditorTreeNodePtrType GetParent() const { return ParentNodePtr; }
 
 	/** Get the display name of this single tree node */
 	UE_API FString GetDisplayString() const;
@@ -143,14 +141,14 @@ public:
 	/** Sets up the delegate for a rename operation */
 	void SetRenameRequestedDelegate(FOnRenameRequested InRenameRequested) { RenameRequestedDelegate = InRenameRequested; }
 	FOnRenameRequested GetRenameRequestedDelegate() { return RenameRequestedDelegate; }
-
+	
+	void ForEachNodeImpl(TFunctionRef<void(FSubobjectEditorTreeNode& Data)> Ftor);
+	void DisposeUnsafe();
 protected: 
 
-	UE_DEPRECATED(5.7, "Use the Parent weak pointer instead.")
+	/** Pointer to the parent of this subobject */
 	FSubobjectEditorTreeNodePtrType ParentNodePtr;
 	
-	TWeakPtr<FSubobjectEditorTreeNode> Parent;
-	
 	// Scope the creation of a node which ends when the initial 'name' is given/accepted by the user, which can be several frames after the node was actually created.
 	TUniquePtr<FScopedTransaction> OngoingCreateTransaction;
 
@@ -802,6 +800,7 @@ protected:
 		static void PopulateHandlesArray(const TArray<FSubobjectEditorTreeNodePtrType>& SlateNodePtrs, TArray<FSubobjectDataHandle>& OutHandles);
 	};
 
+	void ForEachNode(TFunctionRef<void(FSubobjectEditorTreeNode& Data)> Ftor);
 private:
 	FName SelectionSetName;
 };
