Use array of row references for SubobjectData object collections instead of TEDS Hierarchies
* Introduction of FSubobjectDataSubsystemRowReferenceCollection to hold a collection of rows.
* Solves issue where multiple subobject editor instances would stomp hierarchy information on referenced rows.  This issue is due to the same hierarchy information being shared across multiple subobject editor instances and that the relationship would modify common object-referenced rows.

#Synced-CL 48364765

diff --git a/Engine/Source/Editor/LevelEditor/Private/SActorDetails.cpp b/Engine/Source/Editor/LevelEditor/Private/SActorDetails.cpp
--- a/Engine/Source/Editor/LevelEditor/Private/SActorDetails.cpp
+++ b/Engine/Source/Editor/LevelEditor/Private/SActorDetails.cpp
@@ -813,15 +813,18 @@ void SActorDetails::RefreshSubobjectTreeElements(TArrayView<const FSubobjectEdit
 						if (TedsRowAdapter)
 						{
 							RowHandle ComponentRowHandle = TedsRowAdapter->GetRow();
-
-							FHierarchyHandle ReferenceObjectHierarchyName = FSubobjectDataSubsystemContextData_TedsRow::GetReferencedDataHierarchy(*CoreProvider);
-							CoreProvider->WalkDepthFirst(ReferenceObjectHierarchyName, ComponentRowHandle, [&Components](const ICoreProvider& Context, RowHandle Owner, RowHandle Target)
+							
+							if (const FSubobjectDataSubsystemRowReferenceCollection* Collection = CoreProvider->GetColumn<FSubobjectDataSubsystemRowReferenceCollection>(ComponentRowHandle))
 							{
-								if (const Compatibility::FTypedElementColumn* TypedElementColumn = Context.GetColumn<UE::Editor::DataStorage::Compatibility::FTypedElementColumn>(Target))
+								for (UE::Editor::DataStorage::RowHandle Row : Collection->Rows)
 								{
-									Components.Add(TypedElementColumn->Handle);
+									if (const Compatibility::FTypedElementColumn* TypedElementColumn = CoreProvider->GetColumn<UE::Editor::DataStorage::Compatibility::FTypedElementColumn>(Row))
+									{
+										Components.Add(TypedElementColumn->Handle);
+									}
 								}
-							});
+							}
+							
 						}
 					}
 				}
diff --git a/Engine/Source/Editor/SubobjectDataInterface/Private/SubobjectDataAdapterBase.cpp b/Engine/Source/Editor/SubobjectDataInterface/Private/SubobjectDataAdapterBase.cpp
--- a/Engine/Source/Editor/SubobjectDataInterface/Private/SubobjectDataAdapterBase.cpp
+++ b/Engine/Source/Editor/SubobjectDataInterface/Private/SubobjectDataAdapterBase.cpp
@@ -46,12 +46,6 @@ UE::Editor::DataStorage::TableHandle FSubobjectDataSubsystemContextData_TedsRow:
 	return Factory->GetSubobjectDataTableHandle();
 }
 
-UE::Editor::DataStorage::FHierarchyHandle FSubobjectDataSubsystemContextData_TedsRow::GetReferencedDataHierarchy(
-	const UE::Editor::DataStorage::ICoreProvider& DataStorage)
-{
-	return DataStorage.FindHierarchyByName(USubobjectDataTedsFactory::GetReferencedDataHierarchy());
-}
-
 FSubobjectDataSubsystemAdapterBase::~FSubobjectDataSubsystemAdapterBase() = default;
 
 bool FSubobjectDataSubsystemAdapterBase::CanIntercept_GatherSubobjectData(const UE::Editor::FSubobjectEditorContext& Context) const
diff --git a/Engine/Source/Editor/SubobjectDataInterface/Private/SubobjectDataTedsFactory.cpp b/Engine/Source/Editor/SubobjectDataInterface/Private/SubobjectDataTedsFactory.cpp
--- a/Engine/Source/Editor/SubobjectDataInterface/Private/SubobjectDataTedsFactory.cpp
+++ b/Engine/Source/Editor/SubobjectDataInterface/Private/SubobjectDataTedsFactory.cpp
@@ -10,26 +10,13 @@ namespace SubobjectDataTedsFactory::Private
 	FName ReferenceDataHierarchyName = TEXT("SubobjectDataReferencedDataHierarchy");
 }
 
-void USubobjectDataTedsFactory::RegisterHierarchies(UE::Editor::DataStorage::ICoreProvider& DataStorage)
-{
-	UE::Editor::DataStorage::FHierarchyRegistrationParams Params;
-	Params.Name = GetReferencedDataHierarchy();
-	DataStorage.RegisterHierarchy(Params);
-	Super::RegisterHierarchies(DataStorage);
-}
-
 void USubobjectDataTedsFactory::RegisterTables(UE::Editor::DataStorage::ICoreProvider& DataStorage)
 {
-	SubobjectDataTableHandle = DataStorage.RegisterTable<FSubobjectDataTag>(TEXT("SubobjectDataTable"));
+	SubobjectDataTableHandle = DataStorage.RegisterTable<FSubobjectDataTag, FSubobjectDataSubsystemRowReferenceCollection>(TEXT("SubobjectDataTable"));
 	Super::RegisterTables(DataStorage);
 }
 
 UE::Editor::DataStorage::TableHandle USubobjectDataTedsFactory::GetSubobjectDataTableHandle() const
 {
 	return SubobjectDataTableHandle;
-}
-
-FName USubobjectDataTedsFactory::GetReferencedDataHierarchy()
-{
-	return SubobjectDataTedsFactory::Private::ReferenceDataHierarchyName;
-}
+}
\ No newline at end of file
diff --git a/Engine/Source/Editor/SubobjectDataInterface/Private/SubobjectDataTedsFactory.h b/Engine/Source/Editor/SubobjectDataInterface/Private/SubobjectDataTedsFactory.h
--- a/Engine/Source/Editor/SubobjectDataInterface/Private/SubobjectDataTedsFactory.h
+++ b/Engine/Source/Editor/SubobjectDataInterface/Private/SubobjectDataTedsFactory.h
@@ -14,14 +14,9 @@ class USubobjectDataTedsFactory final : public UEditorDataStorageFactory
 public:
 	~USubobjectDataTedsFactory() override = default;
 
-	void RegisterHierarchies(UE::Editor::DataStorage::ICoreProvider& DataStorage) override;
 	void RegisterTables(UE::Editor::DataStorage::ICoreProvider& DataStorage) override;
 
 	UE::Editor::DataStorage::TableHandle GetSubobjectDataTableHandle() const;
-
-	static FName GetReferencedDataHierarchy();
-	
-
 private:
 	UE::Editor::DataStorage::TableHandle SubobjectDataTableHandle = UE::Editor::DataStorage::InvalidTableHandle;
 };
\ No newline at end of file
diff --git a/Engine/Source/Editor/SubobjectDataInterface/Public/SubobjectDataAdapterBase.h b/Engine/Source/Editor/SubobjectDataInterface/Public/SubobjectDataAdapterBase.h
--- a/Engine/Source/Editor/SubobjectDataInterface/Public/SubobjectDataAdapterBase.h
+++ b/Engine/Source/Editor/SubobjectDataInterface/Public/SubobjectDataAdapterBase.h
@@ -53,6 +53,17 @@ struct FSubobjectDataTag : public FEditorDataStorageTag
 	GENERATED_BODY()
 };
 
+// TEDS column to store a collection of rows that are being displayed under a node in the subobject editor component tree
+// Used when SubobjectEditor::UseTedsHierarchies() is false
+USTRUCT()
+struct FSubobjectDataSubsystemRowReferenceCollection : public FEditorDataStorageColumn
+{
+	GENERATED_BODY()
+	
+	UPROPERTY()
+	TArray<FTedsRowHandle> Rows;
+};
+
 // An implementation of the ContextData that supports containing a single UObject
 // This is to provide interoperability with the existing ContextObject used in APIs of the SubobjectDataSubsystem
 USTRUCT()
@@ -84,10 +95,6 @@ FSubobjectDataSubsystemContextData_TedsRow : public FSubobjectDataSubsystemConte
 
 	UE_API static UE::Editor::DataStorage::TableHandle GetSubobjectItemTable(const UE::Editor::DataStorage::ICoreProvider& DataStorage);
 	
-	// Can be used along with TEDS API to walk any data rows that are referenced by this SubobjectData
-	// For example, a SubobjectData could represent multiple UActorComponent instances, they would be represented
-	// as children in the ReferencedDataHierarchy
-	UE_API static UE::Editor::DataStorage::FHierarchyHandle GetReferencedDataHierarchy(const UE::Editor::DataStorage::ICoreProvider& DataStorage);
 private:
 
 	TSharedPtr<UE::Editor::DataStorage::FScopedRow> Row;
diff --git a/Engine/Source/Editor/SubobjectEditor/Private/SSubobjectEditor.cpp b/Engine/Source/Editor/SubobjectEditor/Private/SSubobjectEditor.cpp
--- a/Engine/Source/Editor/SubobjectEditor/Private/SSubobjectEditor.cpp
+++ b/Engine/Source/Editor/SubobjectEditor/Private/SSubobjectEditor.cpp
@@ -3446,19 +3446,22 @@ void SSubobjectEditor::GetSelectedItemsForContextMenu(TArray<FComponentEventCons
 				using namespace UE::Editor::DataStorage;
 				ICoreProvider* CoreProvider = UE::Editor::DataStorage::GetMutableDataStorageFeature<ICoreProvider>(StorageFeatureName);
 				check(CoreProvider);
-				
-				FHierarchyHandle ReferenceObjectHierarchyName = FSubobjectDataSubsystemContextData_TedsRow::GetReferencedDataHierarchy(*CoreProvider);
+
 				RowHandle ComponentRowHandle = AsTedsRow->GetRow();
-				CoreProvider->WalkDepthFirst(ReferenceObjectHierarchyName, ComponentRowHandle, [&OutSelectedItems](const ICoreProvider& Context, RowHandle Owner, RowHandle Target)
+				if (const FSubobjectDataSubsystemRowReferenceCollection* Collection = CoreProvider->GetColumn<FSubobjectDataSubsystemRowReferenceCollection>(ComponentRowHandle))
 				{
-					const FTypedElementUObjectColumn* UObjectColumn = Context.GetColumn<FTypedElementUObjectColumn>(Target);
-					if (UObjectColumn)
+					for (UE::Editor::DataStorage::RowHandle Row : Collection->Rows)
 					{
-						UObject* Object = UObjectColumn->Object.Get();
-						FComponentEventConstructionData& ContextItem = OutSelectedItems.AddDefaulted_GetRef();
-						ContextItem.Component = Object;
+						const FTypedElementUObjectColumn* UObjectColumn = CoreProvider->GetColumn<FTypedElementUObjectColumn>(Row);
+						if (UObjectColumn)
+						{
+							UObject* Object = UObjectColumn->Object.Get();
+							FComponentEventConstructionData& ContextItem = OutSelectedItems.AddDefaulted_GetRef();
+							ContextItem.Component = Object;
+						}
 					}
-				});
+				}
+				
 			}
 			else if (AsSingleUObject)
 			{
