Fix newly created entities not getting a populated subobject editor
* Issue stems from the TEDS system not populating object row immediately, instead they are deferred until end of frame.  The object rows for the entity or components effectively have no data when the subobject editor's GatherSubobjectData function is called.  If a query is built to pull objects from TEDS that have a specific FTypedElementSelectionColumn, it will appear as if there are no entries and thus an empty selection set.
* The fix is to use the TypedElement SelectionSet for now and pull the objects off it it based on the Object interface.  TypedElements better conform to the lifetime of the objects they represent on the front-end side of things (ie. they are created when the objects are created).  Confusingly, they are defer destroyed so they actually mirror the lifetime worse on the backend - though this doesn't matter for this case since the TypedElements are used to just pull TEDS row handles.
* There is still an outstanding issue though where a newly created entity will also have a component automatically selected (ie. often transform_component).  Since the component's row is also pending add, it is not possible to get the UObject* to pass to the details panel and therefore the details panel will remain empty.  A workaround to de-select and reselect exists for now or to add a component.

#Synced-CL 48738172

diff --git a/Engine/Source/Editor/LevelEditor/Private/SActorDetails.cpp b/Engine/Source/Editor/LevelEditor/Private/SActorDetails.cpp
--- a/Engine/Source/Editor/LevelEditor/Private/SActorDetails.cpp
+++ b/Engine/Source/Editor/LevelEditor/Private/SActorDetails.cpp
@@ -505,6 +505,8 @@ void SActorDetails::Construct(
 		.OnSelectionUpdated(this, &SActorDetails::OnSubobjectEditorTreeViewSelectionChanged)
 		.OnItemDoubleClicked(this, &SActorDetails::OnSubobjectEditorTreeViewItemDoubleClicked)
 		.SCSEditorUICustomization(SubobjectEditorCustomization);
+	
+	SubobjectEditor->SetSelectionSet(SelectionSet);
 
 	ComponentsBox->SetContent(SubobjectEditor.ToSharedRef());
 
diff --git a/Engine/Source/Editor/SubobjectDataInterface/Public/SubobjectDataSubsystem.h b/Engine/Source/Editor/SubobjectDataInterface/Public/SubobjectDataSubsystem.h
--- a/Engine/Source/Editor/SubobjectDataInterface/Public/SubobjectDataSubsystem.h
+++ b/Engine/Source/Editor/SubobjectDataInterface/Public/SubobjectDataSubsystem.h
@@ -7,7 +7,9 @@
 #include "Subsystems/EngineSubsystem.h"
 #include "SubobjectData.h"
 #include "ModuleDescriptor.h"
+#include "Elements/Framework/TypedElementSelectionSet.h"
 #include "SubobjectDataAdapterBase.h"
+
 #include "SubobjectDataSubsystem.generated.h"
 
 #define UE_API SUBOBJECTDATAINTERFACE_API
@@ -101,8 +103,8 @@ namespace UE::Editor
 		//  eventually be phased out in favor of always using SelectionSetName or
 		//  some other way to identify the target objects (since there may be multiple)
 		UObject* ObjectContext = nullptr;
-		// If editor operates on a selection set, its identifier
-		FName SelectionSetName;
+		
+		TWeakObjectPtr<UTypedElementSelectionSet> SelectionSet = nullptr;
 	};
 }
 
diff --git a/Engine/Source/Editor/SubobjectDataInterface/SubobjectDataInterface.Build.cs b/Engine/Source/Editor/SubobjectDataInterface/SubobjectDataInterface.Build.cs
--- a/Engine/Source/Editor/SubobjectDataInterface/SubobjectDataInterface.Build.cs
+++ b/Engine/Source/Editor/SubobjectDataInterface/SubobjectDataInterface.Build.cs
@@ -13,6 +13,7 @@ public class SubobjectDataInterface : ModuleRules
 			{
 				"Core",
 				"CoreUObject",
+				"TypedElementFramework"
 			}
 		);
 
diff --git a/Engine/Source/Editor/SubobjectEditor/Private/SSubobjectEditor.cpp b/Engine/Source/Editor/SubobjectEditor/Private/SSubobjectEditor.cpp
--- a/Engine/Source/Editor/SubobjectEditor/Private/SSubobjectEditor.cpp
+++ b/Engine/Source/Editor/SubobjectEditor/Private/SSubobjectEditor.cpp
@@ -2365,7 +2365,7 @@ EVisibility SSubobjectEditor::GetDesiredVisibility() const
 	EVisibility CustomizationVisibility = EVisibility::Collapsed;
 	UE::Editor::FSubobjectEditorContext Context;
 	Context.ObjectContext = GetObjectContext();
-	Context.SelectionSetName = SelectionSetName;
+	Context.SelectionSet = SelectionSet;
 	if (UICustomization && UICustomization->OverrideDesiredVisibility(Context, CustomizationVisibility))
 	{
 		return CustomizationVisibility;
@@ -2578,7 +2578,7 @@ void SSubobjectEditor::UpdateTree(bool bRegenerateTreeNodes /* = true */)
 
 			UE::Editor::FSubobjectEditorContext Params;
 			Params.ObjectContext = GetObjectContext();
-			Params.SelectionSetName = SelectionSetName;
+			Params.SelectionSet = SelectionSet;
 			DataSubsystem->GatherSubobjectData(Params, SubobjectData);
 
 			if (UICustomization)
@@ -2762,9 +2762,9 @@ FReply SSubobjectEditor::OnKeyDown(const FGeometry& MyGeometry, const FKeyEvent&
 	return FReply::Unhandled();
 }
 
-void SSubobjectEditor::SetSelectionSetName(FName SelectionSetNameIn)
+void SSubobjectEditor::SetSelectionSet(UTypedElementSelectionSet* InSelectionSet)
 {
-	SelectionSetName = SelectionSetNameIn;
+	SelectionSet = InSelectionSet;
 }
 
 FReply SSubobjectEditor::TryHandleAssetDragDropOperation(const FDragDropEvent& DragDropEvent)
diff --git a/Engine/Source/Editor/SubobjectEditor/Public/SSubobjectEditor.h b/Engine/Source/Editor/SubobjectEditor/Public/SSubobjectEditor.h
--- a/Engine/Source/Editor/SubobjectEditor/Public/SSubobjectEditor.h
+++ b/Engine/Source/Editor/SubobjectEditor/Public/SSubobjectEditor.h
@@ -466,11 +466,7 @@ protected:
 	UE_API virtual FReply OnKeyDown(const FGeometry& MyGeometry, const FKeyEvent& InKeyEvent) override;
 	
 public:
-	/** 
-	 * Sets the identifier representing this subobject editor's selection in TEDS.
-	 *  If not specified, will use the level editor selection. 
-	 */
-	UE_API void SetSelectionSetName(FName SelectionSetNameIn);
+	UE_API void SetSelectionSet(UTypedElementSelectionSet* InSelectionSet);
 
 	/** Select the given tree node */
 	UE_API void SelectNode(FSubobjectEditorTreeNodePtrType InNodeToSelect, bool bIsCntrlDown);
@@ -803,7 +799,7 @@ protected:
 
 	void ForEachNode(TFunctionRef<void(FSubobjectEditorTreeNode& Data)> Ftor);
 private:
-	FName SelectionSetName;
+	TWeakObjectPtr<UTypedElementSelectionSet> SelectionSet;
 };
 
 #undef UE_API
diff --git a/Engine/Source/Runtime/TypedElementRuntime/Private/Elements/Framework/TypedElementSelectionSet.cpp b/Engine/Source/Runtime/TypedElementRuntime/Private/Elements/Framework/TypedElementSelectionSet.cpp
--- a/Engine/Source/Runtime/TypedElementRuntime/Private/Elements/Framework/TypedElementSelectionSet.cpp
+++ b/Engine/Source/Runtime/TypedElementRuntime/Private/Elements/Framework/TypedElementSelectionSet.cpp
@@ -403,6 +403,11 @@ void UTypedElementSelectionSet::SetNameForTedsIntegration(FName InNameForIntegra
 	ListNameForTedsIntegration = InNameForIntegration;
 }
 
+FName UTypedElementSelectionSet::GetNameForTedsIntegration() const
+{
+	return ListNameForTedsIntegration;
+}
+
 FTypedElementHandle UTypedElementSelectionSet::GetSelectionElement(const FTypedElementHandle& InElementHandle, const ETypedElementSelectionMethod InSelectionMethod) const
 {
 	FTypedElementSelectionSetElement SelectionSetElement = ResolveSelectionSetElement(InElementHandle);
diff --git a/Engine/Source/Runtime/TypedElementRuntime/Public/Elements/Framework/TypedElementSelectionSet.h b/Engine/Source/Runtime/TypedElementRuntime/Public/Elements/Framework/TypedElementSelectionSet.h
--- a/Engine/Source/Runtime/TypedElementRuntime/Public/Elements/Framework/TypedElementSelectionSet.h
+++ b/Engine/Source/Runtime/TypedElementRuntime/Public/Elements/Framework/TypedElementSelectionSet.h
@@ -260,6 +260,8 @@ public:
 	 * Sets the name to use for teds integration so that it can be reapplied on undo/redo
 	 */
 	TYPEDELEMENTRUNTIME_API void SetNameForTedsIntegration(const FName InNameForIntegration);
+	
+	TYPEDELEMENTRUNTIME_API FName GetNameForTedsIntegration() const;
 
 	/**
 	 * Given an element, return the element that should actually perform a selection operation.
