Updated entity components to display their class display name (which is the unmangled class name) along with their full Verse path in their tooltips.

#Synced-CL 49041953

diff --git a/Engine/Source/Editor/SubobjectDataInterface/Private/SubobjectData.cpp b/Engine/Source/Editor/SubobjectDataInterface/Private/SubobjectData.cpp
--- a/Engine/Source/Editor/SubobjectDataInterface/Private/SubobjectData.cpp
+++ b/Engine/Source/Editor/SubobjectDataInterface/Private/SubobjectData.cpp
@@ -481,14 +481,7 @@ FString FSubobjectData::GetDisplayString(bool bShowNativeComponentNames /* = tru
 		FString UnnamedString = LOCTEXT("UnnamedToolTip", "Unnamed").ToString();
 		FString NativeString = IsNativeComponent() ? LOCTEXT("NativeToolTip", "Native ").ToString() : TEXT("");
 
-		if (ComponentTemplate != nullptr)
-		{
-			return FString::Printf(TEXT("[%s %s%s]"), *UnnamedString, *NativeString, *ComponentTemplate->GetClass()->GetName());
-		}
-		else
-		{
-			return FString::Printf(TEXT("[%s %s]"), *UnnamedString, *NativeString);
-		}
+		return FString::Printf(TEXT("[%s %s]"), *UnnamedString, *NativeString);
 	}
 }
 
@@ -820,6 +813,24 @@ bool FSubobjectData::IsAssetVisible() const
 	return false;
 }
 
+FText FSubobjectData::GetToolTipText() const
+{
+	REDIRECT_TO_ADAPTER(GetToolTipText)
+
+	FText NodeDisplayName = FText::FromString(GetDisplayString());
+
+	if (const UActorComponent* ComponentTemplate = GetComponentTemplate())
+	{
+		FFormatNamedArguments Args;
+		Args.Add(TEXT("NodeName"), NodeDisplayName);
+		Args.Add(TEXT("ClassName"), GetComponentTemplate()->GetClass()->GetDisplayNameText());
+
+		return FText::Format(LOCTEXT("ComponentTooltip", "{NodeName} ({ClassName})"), Args);
+	}
+
+	return NodeDisplayName;
+}
+
 FText FSubobjectData::GetMobilityToolTipText() const
 {
 	REDIRECT_TO_ADAPTER(GetMobilityToolTipText)
diff --git a/Engine/Source/Editor/SubobjectDataInterface/Private/SubobjectDataAdapterBase.cpp b/Engine/Source/Editor/SubobjectDataInterface/Private/SubobjectDataAdapterBase.cpp
--- a/Engine/Source/Editor/SubobjectDataInterface/Private/SubobjectDataAdapterBase.cpp
+++ b/Engine/Source/Editor/SubobjectDataInterface/Private/SubobjectDataAdapterBase.cpp
@@ -370,6 +370,11 @@ bool FSubobjectDataAdapterBase::IsAssetVisible(const FSubobjectData& Data) const
 	return false;
 }
 
+FText FSubobjectDataAdapterBase::GetToolTipText(const FSubobjectData& Data) const
+{
+	return FText();
+}
+
 FText FSubobjectDataAdapterBase::GetMobilityToolTipText(const FSubobjectData& Data) const
 {
 	return FText();
diff --git a/Engine/Source/Editor/SubobjectDataInterface/Public/SubobjectData.h b/Engine/Source/Editor/SubobjectDataInterface/Public/SubobjectData.h
--- a/Engine/Source/Editor/SubobjectDataInterface/Public/SubobjectData.h
+++ b/Engine/Source/Editor/SubobjectDataInterface/Public/SubobjectData.h
@@ -208,6 +208,13 @@ public:
 	/** Returns true if the asset this object represents is visible */
 	UE_API bool IsAssetVisible() const;
 
+	/**
+	* Retrieves tooltip text describing the specified component.
+	*
+	* @returns An FText object containing a description of the component
+	*/
+	UE_API FText GetToolTipText() const;
+
 	/**
 	* Retrieves tooltip text describing the specified component's mobility.
 	*
diff --git a/Engine/Source/Editor/SubobjectDataInterface/Public/SubobjectDataAdapterBase.h b/Engine/Source/Editor/SubobjectDataInterface/Public/SubobjectDataAdapterBase.h
--- a/Engine/Source/Editor/SubobjectDataInterface/Public/SubobjectDataAdapterBase.h
+++ b/Engine/Source/Editor/SubobjectDataInterface/Public/SubobjectDataAdapterBase.h
@@ -202,6 +202,7 @@ public:
 	UE_API virtual FText GetAssetName(const FSubobjectData& Data) const;
 	UE_API virtual FText GetAssetPath(const FSubobjectData& Data) const;
 	UE_API virtual bool IsAssetVisible(const FSubobjectData& Data) const;
+	UE_API virtual FText GetToolTipText(const FSubobjectData& Data) const;
 	UE_API virtual FText GetMobilityToolTipText(const FSubobjectData& Data) const;
 	UE_API virtual FText GetComponentEditorOnlyTooltipText(const FSubobjectData& Data) const;
 	UE_API virtual FText GetIntroducedInToolTipText(const FSubobjectData& Data) const;
diff --git a/Engine/Source/Editor/SubobjectEditor/Private/SSubobjectEditor.cpp b/Engine/Source/Editor/SubobjectEditor/Private/SSubobjectEditor.cpp
--- a/Engine/Source/Editor/SubobjectEditor/Private/SSubobjectEditor.cpp
+++ b/Engine/Source/Editor/SubobjectEditor/Private/SSubobjectEditor.cpp
@@ -1575,7 +1575,7 @@ TSharedRef<SToolTip> SSubobject_RowWidget::CreateComponentTooltipWidget(const FS
 				[
 					SNew(STextBlock)
 					.TextStyle(FAppStyle::Get(), "SCSEditor.ComponentTooltip.Title")
-					.Text(this, &SSubobject_RowWidget::GetTooltipText)
+					.Text(this, &SSubobject_RowWidget::GetToolTipText)
 				]
 			]
 		]
@@ -1592,7 +1592,7 @@ TSharedRef<SToolTip> SSubobject_RowWidget::CreateComponentTooltipWidget(const FS
 		]
 	];
 
-	return IDocumentation::Get()->CreateToolTip(TAttribute<FText>(this, &SSubobject_RowWidget::GetTooltipText),
+	return IDocumentation::Get()->CreateToolTip(TAttribute<FText>(this, &SSubobject_RowWidget::GetToolTipText),
 	                                            TooltipContent, InfoBox, GetDocumentationLink(),
 	                                            GetDocumentationExcerptName());
 }
@@ -1655,7 +1655,7 @@ TSharedRef<SToolTip> SSubobject_RowWidget::CreateActorTooltipWidget(const FSubob
 	return IDocumentation::Get()->CreateToolTip(TAttribute<FText>(this, &SSubobject_RowWidget::GetActorDisplayText), TooltipContent, InfoBox, TEXT(""), TEXT(""));
 }
 
-FText SSubobject_RowWidget::GetTooltipText() const
+FText SSubobject_RowWidget::GetToolTipText() const
 {
 	const FSubobjectData* Data = GetSubobjectPtr()->GetDataSource();
 
@@ -1687,16 +1687,7 @@ FText SSubobject_RowWidget::GetTooltipText() const
 	}
 	else
 	{
-		const UClass* Class = (Data->GetComponentTemplate() != nullptr)
-			                      ? Data->GetComponentTemplate()->GetClass()
-			                      : nullptr;
-		const FText ClassDisplayName = FBlueprintEditorUtils::GetFriendlyClassDisplayName(Class);
-
-		FFormatNamedArguments Args;
-		Args.Add(TEXT("ClassName"), ClassDisplayName);
-		Args.Add(TEXT("NodeName"), FText::FromString(Data->GetDisplayString()));
-
-		return FText::Format(LOCTEXT("ComponentTooltip", "{NodeName} ({ClassName})"), Args);
+		return Data->GetToolTipText();
 	}
 }
 
diff --git a/Engine/Source/Editor/SubobjectEditor/Public/SSubobjectEditor.h b/Engine/Source/Editor/SubobjectEditor/Public/SSubobjectEditor.h
--- a/Engine/Source/Editor/SubobjectEditor/Public/SSubobjectEditor.h
+++ b/Engine/Source/Editor/SubobjectEditor/Public/SSubobjectEditor.h
@@ -308,7 +308,7 @@ private:
 	/** Create the tooltip for actor subobjects */
 	TSharedRef<SToolTip> CreateActorTooltipWidget(const FSubobjectEditorTreeNodePtrType& InNode) const;
 	
-	FText GetTooltipText() const;
+	FText GetToolTipText() const;
 	
 	FText GetActorClassNameText() const;
 	FText GetActorSuperClassNameText() const;
