SSubobjectEditor::FindSlateNodeForHandle: find slate node for FSubobjectData that doesn't have a WeakObjectPtr but does have a corresponding teds row

#Synced-CL 50000013

#Synced-CL 50000013

diff --git a/Engine/Source/Editor/SubobjectEditor/Private/SSubobjectEditor.cpp b/Engine/Source/Editor/SubobjectEditor/Private/SSubobjectEditor.cpp
--- a/Engine/Source/Editor/SubobjectEditor/Private/SSubobjectEditor.cpp
+++ b/Engine/Source/Editor/SubobjectEditor/Private/SSubobjectEditor.cpp
@@ -3144,21 +3144,68 @@ FSubobjectEditorTreeNodePtrType SSubobjectEditor::FindSlateNodeForHandle(const F
 	FSubobjectEditorTreeNodePtrType OutNodePtr;
 	if (InHandleToFind.IsValid() && RootNodes.Num() > 0)
 	{
-		FSubobjectData* DataToFind = InHandleToFind.GetSharedDataPtr().Get();
-		
-		TSet<FSubobjectEditorTreeNodePtrType> VisitedNodes;
-		DepthFirstTraversal(RootNodes[0], VisitedNodes,
-            [&OutNodePtr, &DataToFind](
-            const FSubobjectEditorTreeNodePtrType& CurNodePtr)
-            {
-            	if(CurNodePtr->GetDataHandle().IsValid())
-            	{
-            		if(CurNodePtr->GetDataHandle().GetSharedDataPtr()->GetObject() == DataToFind->GetObject())
-            		{
-                        OutNodePtr = CurNodePtr;
-                    }
-            	}
-            });
+		const FSubobjectDataSubsystemContextData_SingleUObjectContextObject* const AsSingleUObject = InHandleToFind.GetData()->GetContextData().GetPtr<FSubobjectDataSubsystemContextData_SingleUObjectContextObject>();
+		const FSubobjectDataSubsystemContextData_TedsRow* const AsTedsRow = InHandleToFind.GetData()->GetContextData().GetPtr<FSubobjectDataSubsystemContextData_TedsRow>();
+
+		if (AsSingleUObject)
+		{
+			FSubobjectData* DataToFind = InHandleToFind.GetSharedDataPtr().Get();
+			TSet<FSubobjectEditorTreeNodePtrType> VisitedNodes;
+			DepthFirstTraversal(RootNodes[0], VisitedNodes,
+				[&OutNodePtr, &DataToFind](
+					const FSubobjectEditorTreeNodePtrType& CurNodePtr)
+				{
+					if (CurNodePtr->GetDataHandle().IsValid())
+					{
+						if (CurNodePtr->GetDataHandle().GetSharedDataPtr()->GetObject() == DataToFind->GetObject())
+						{
+							OutNodePtr = CurNodePtr;
+						}
+					}
+				});
+		}
+		else if (AsTedsRow)
+		{
+			using namespace UE::Editor::DataStorage;
+			const ICoreProvider* const CoreProvider = GetMutableDataStorageFeature<ICoreProvider>(StorageFeatureName);
+			check(CoreProvider);
+
+			auto GetObjectFromTedsRow = [CoreProvider](RowHandle Row) -> UObject*
+			{
+				if (const FSubobjectDataSubsystemRowReferenceCollection* const Collection = CoreProvider->GetColumn<FSubobjectDataSubsystemRowReferenceCollection>(Row))
+				{
+					for (const RowHandle CollectionRow : Collection->Rows)
+					{
+						const FTypedElementUObjectColumn* const UObjectColumn = CoreProvider->GetColumn<FTypedElementUObjectColumn>(CollectionRow);
+						if (UObjectColumn)
+						{
+							return UObjectColumn->Object.Get();
+						}
+					}
+				}
+				return nullptr;
+			};
+
+			if (const UObject* const QueryObject = GetObjectFromTedsRow(AsTedsRow->GetRow()))
+			{
+				TSet<FSubobjectEditorTreeNodePtrType> VisitedNodes;
+				DepthFirstTraversal(RootNodes[0], VisitedNodes, [&OutNodePtr, QueryObject, GetObjectFromTedsRow, CoreProvider](const FSubobjectEditorTreeNodePtrType& CurNodePtr)
+				{
+					if (CurNodePtr->GetDataHandle().IsValid())
+					{
+						if (const FSubobjectDataSubsystemContextData_TedsRow* const TestTedsRow = CurNodePtr->GetDataHandle().GetData()->GetContextData().GetPtr<FSubobjectDataSubsystemContextData_TedsRow>())
+						{
+							const UObject* const TestObject = GetObjectFromTedsRow(TestTedsRow->GetRow());
+							if (QueryObject == TestObject)
+							{
+								OutNodePtr = CurNodePtr;
+							}
+						}
+					}
+				});
+			}
+		}
+
 	}
 
 	return OutNodePtr;
