[EntityProxyActor Removal]
* Enable adding new component type and existing component type to selected entities when EntityProxyActors disabled.
* Replace simulation of new selection on component add with actor agnostic version
* Remove a couple unused functions
* Add accessor to SubobjectEditor's configured SelectionSet

#Synced-CL 50990257

diff --git a/Engine/Source/Editor/SubobjectEditor/Private/SSubobjectEditor.cpp b/Engine/Source/Editor/SubobjectEditor/Private/SSubobjectEditor.cpp
--- a/Engine/Source/Editor/SubobjectEditor/Private/SSubobjectEditor.cpp
+++ b/Engine/Source/Editor/SubobjectEditor/Private/SSubobjectEditor.cpp
@@ -2759,6 +2759,11 @@ void SSubobjectEditor::SetSelectionSet(UTypedElementSelectionSet* InSelectionSet
 	SelectionSet = InSelectionSet;
 }
 
+UTypedElementSelectionSet* SSubobjectEditor::GetSelectionSet() const
+{
+	return SelectionSet.Get();
+}
+
 FReply SSubobjectEditor::TryHandleAssetDragDropOperation(const FDragDropEvent& DragDropEvent)
 {
 	TSharedPtr<FDragDropOperation> Operation = DragDropEvent.GetOperation();
diff --git a/Engine/Source/Editor/SubobjectEditor/Public/SSubobjectEditor.h b/Engine/Source/Editor/SubobjectEditor/Public/SSubobjectEditor.h
--- a/Engine/Source/Editor/SubobjectEditor/Public/SSubobjectEditor.h
+++ b/Engine/Source/Editor/SubobjectEditor/Public/SSubobjectEditor.h
@@ -467,6 +467,8 @@ protected:
 	
 public:
 	UE_API void SetSelectionSet(UTypedElementSelectionSet* InSelectionSet);
+	
+	UE_API UTypedElementSelectionSet* GetSelectionSet() const;
 
 	/** Select the given tree node */
 	UE_API void SelectNode(FSubobjectEditorTreeNodePtrType InNodeToSelect, bool bIsCntrlDown);
