Added a TypedElement customization for the Subobject Editor, allowing correct component duplication

#Synced-CL 52103551

diff --git a/Engine/Source/Editor/Kismet/Private/BlueprintEditor.cpp b/Engine/Source/Editor/Kismet/Private/BlueprintEditor.cpp
--- a/Engine/Source/Editor/Kismet/Private/BlueprintEditor.cpp
+++ b/Engine/Source/Editor/Kismet/Private/BlueprintEditor.cpp
@@ -185,6 +185,7 @@
 #include "AssetEditorModeManager.h"
 #include "Editor/EditorEngine.h"
 #include "EditorViewportSelectabilityBridge.h"
+#include "Elements/Component/ComponentElementSubobjectEditorCommonActionsCustomization.h"
 #include "IBlueprintSelectionNotifier.h"
 #include "Tools/EdModeInteractiveToolsContext.h"
 
@@ -2178,6 +2179,19 @@ void FBlueprintEditor::CommonInitialization(const TArray<UBlueprint*>& InitBluep
 
 	CreateDefaultTabContents(InitBlueprints);
 
+	// Only register our component common actions customization in Standalone mode.
+	if (!IsWorldCentricAssetEditor())
+	{
+		const TSharedPtr<IToolkitHost> PinnedToolkitHost = ToolkitHost.Pin();
+		if (ensure(PinnedToolkitHost.IsValid()))
+		{
+			TUniquePtr<FComponentElementSubobjectEditorCommonActionsCustomization> ComponentCustomization = MakeUnique<FComponentElementSubobjectEditorCommonActionsCustomization>();
+			ComponentCustomization->SetSubobjectEditor(GetSubobjectEditor());
+
+			PinnedToolkitHost->GetCommonActions()->RegisterInterfaceCustomizationByTypeName(NAME_Components, MoveTemp(ComponentCustomization));
+		}
+	}
+
 	FCoreUObjectDelegates::OnPreObjectPropertyChanged.AddSP(this, &FBlueprintEditor::OnPreObjectPropertyChanged);
 	FCoreUObjectDelegates::OnObjectPropertyChanged.AddSP(this, &FBlueprintEditor::OnPostObjectPropertyChanged);
 
diff --git a/Engine/Source/Editor/SubobjectEditor/Private/Elements/Component/ComponentElementSubobjectEditorCommonActionsCustomization.cpp b/Engine/Source/Editor/SubobjectEditor/Private/Elements/Component/ComponentElementSubobjectEditorCommonActionsCustomization.cpp
new file mode 100644
--- /dev/null
+++ b/Engine/Source/Editor/SubobjectEditor/Private/Elements/Component/ComponentElementSubobjectEditorCommonActionsCustomization.cpp
@@ -0,0 +1,37 @@
+// Copyright Epic Games, Inc. All Rights Reserved.
+
+#include "Elements/Component/ComponentElementSubobjectEditorCommonActionsCustomization.h"
+#include "SSubobjectEditor.h"
+#include "Elements/Framework/EngineElementsLibrary.h"
+#include "SubobjectData.h"
+
+void FComponentElementSubobjectEditorCommonActionsCustomization::DuplicateElements(ITypedElementWorldInterface* InWorldInterface, TArrayView<const FTypedElementHandle> InElementHandles, UWorld* InWorld, const FVector& InLocationOffset, TArray<FTypedElementHandle>& OutNewElements)
+{
+	if (const TSharedPtr<SSubobjectEditor> Editor = GetSubobjectEditor())
+	{
+		if (Editor->CanDuplicateComponent())
+		{
+			Editor->OnDuplicateComponent();
+
+			// OnDuplicateComponent selects the newly duplicated nodes. Convert
+			// the current selection to typed element handles so callers can
+			// identify what was created.
+			for (const FSubobjectDataHandle& Handle : Editor->GetSelectedHandles())
+			{
+				if (const FSubobjectData* Data = Handle.GetData())
+				{
+					if (const UActorComponent* Component = Data->GetComponentTemplate())
+					{
+						if (FTypedElementHandle ElementHandle = UEngineElementsLibrary::AcquireEditorComponentElementHandle(Component))
+						{
+							OutNewElements.Add(MoveTemp(ElementHandle));
+						}
+					}
+				}
+			}
+			return;
+		}
+	}
+
+	Super::DuplicateElements(InWorldInterface, InElementHandles, InWorld, InLocationOffset, OutNewElements);
+}
diff --git a/Engine/Source/Editor/SubobjectEditor/Public/Elements/Component/ComponentElementSubobjectEditorCommonActionsCustomization.h b/Engine/Source/Editor/SubobjectEditor/Public/Elements/Component/ComponentElementSubobjectEditorCommonActionsCustomization.h
new file mode 100644
--- /dev/null
+++ b/Engine/Source/Editor/SubobjectEditor/Public/Elements/Component/ComponentElementSubobjectEditorCommonActionsCustomization.h
@@ -0,0 +1,18 @@
+// Copyright Epic Games, Inc. All Rights Reserved.
+
+#pragma once
+
+#include "Elements/Framework/TypedElementCommonActions.h"
+#include "Elements/Framework/TypedElementSubobjectEditorMixin.h"
+
+#define UE_API SUBOBJECTEDITOR_API
+
+class FComponentElementSubobjectEditorCommonActionsCustomization : public FTypedElementCommonActionsCustomization, public FTypedElementSubobjectEditorMixin
+{
+	using Super = FTypedElementCommonActionsCustomization;
+
+public:
+	UE_API virtual void DuplicateElements(ITypedElementWorldInterface* InWorldInterface, TArrayView<const FTypedElementHandle> InElementHandles, UWorld* InWorld, const FVector& InLocationOffset, TArray<FTypedElementHandle>& OutNewElements) override;
+};
+
+#undef UE_API
diff --git a/Engine/Source/Editor/SubobjectEditor/Public/Elements/Framework/TypedElementSubobjectEditorMixin.h b/Engine/Source/Editor/SubobjectEditor/Public/Elements/Framework/TypedElementSubobjectEditorMixin.h
new file mode 100644
--- /dev/null
+++ b/Engine/Source/Editor/SubobjectEditor/Public/Elements/Framework/TypedElementSubobjectEditorMixin.h
@@ -0,0 +1,34 @@
+// Copyright Epic Games, Inc. All Rights Reserved.
+
+#pragma once
+
+#include "Templates/SharedPointer.h"
+
+class SSubobjectEditor;
+
+/**
+ * Mixin for subobject editor customizations that provides access to the subobject editor widget.
+ */
+class FTypedElementSubobjectEditorMixin
+{
+public:
+	/**
+	 * Get the subobject editor associated with this customization.
+	 */
+	TSharedPtr<SSubobjectEditor> GetSubobjectEditor() const
+	{
+		return SubobjectEditorWeak.Pin();
+	}
+
+	/**
+	 * Set the subobject editor associated with this customization.
+	 */
+	void SetSubobjectEditor(const TSharedPtr<SSubobjectEditor>& InSubobjectEditor)
+	{
+		SubobjectEditorWeak = InSubobjectEditor;
+	}
+
+private:
+	/** The subobject editor associated with this customization. */
+	TWeakPtr<SSubobjectEditor> SubobjectEditorWeak;
+};
diff --git a/Engine/Source/Editor/SubobjectEditor/Public/SSubobjectEditor.h b/Engine/Source/Editor/SubobjectEditor/Public/SSubobjectEditor.h
--- a/Engine/Source/Editor/SubobjectEditor/Public/SSubobjectEditor.h
+++ b/Engine/Source/Editor/SubobjectEditor/Public/SSubobjectEditor.h
@@ -632,6 +632,8 @@ public:
 	virtual void OnDeleteNodes() = 0;
 	
 protected:
+	friend class FComponentElementSubobjectEditorCommonActionsCustomization;
+
 	UE_API bool CanDuplicateComponent() const;
 	virtual void OnDuplicateComponent() = 0;
 
