EditorUsability : SubobjectEditor
- Fixed Blueprint Editor Drag and Drop of asset in tree hierarchy does not match user intent

#Synced-CL 52419302

diff --git a/Engine/Source/Editor/SubobjectEditor/Private/SSubobjectEditor.cpp b/Engine/Source/Editor/SubobjectEditor/Private/SSubobjectEditor.cpp
--- a/Engine/Source/Editor/SubobjectEditor/Private/SSubobjectEditor.cpp
+++ b/Engine/Source/Editor/SubobjectEditor/Private/SSubobjectEditor.cpp
@@ -1436,8 +1436,7 @@ FReply SSubobject_RowWidget::HandleOnAcceptDrop(const FDragDropEvent& DragDropEv
 		TSharedPtr<SSubobjectEditorDragDropTree> TreeWidget = PinnedEditor.IsValid() ? PinnedEditor->GetDragDropTree() : nullptr;
 		if (TreeWidget.IsValid())
 		{
-			// The widget geometry is irrelevant to the tree widget's OnDrop
-			PinnedEditor->GetDragDropTree()->OnDrop(FGeometry(), DragDropEvent);
+			return PinnedEditor->TryHandleAssetDragDropOperation(DragDropEvent, TargetItem);
 		}
 	}
 	
@@ -2797,7 +2796,12 @@ UTypedElementSelectionSet* SSubobjectEditor::GetSelectionSet() const
 
 FReply SSubobjectEditor::TryHandleAssetDragDropOperation(const FDragDropEvent& DragDropEvent)
 {
-	TSharedPtr<FDragDropOperation> Operation = DragDropEvent.GetOperation();
+	return TryHandleAssetDragDropOperation(DragDropEvent, nullptr);
+}
+
+FReply SSubobjectEditor::TryHandleAssetDragDropOperation(const FDragDropEvent& InDragDropEvent, FSubobjectEditorTreeNodePtrType InTargetItem)
+{
+	TSharedPtr<FDragDropOperation> Operation = InDragDropEvent.GetOperation();
 	if (Operation.IsValid() && (Operation->IsOfType<FExternalDragOperation>() || Operation->IsOfType<FAssetDragDropOp>()))
 	{
 		TArray< FAssetData > DroppedAssetData = AssetUtil::ExtractAssetDataFromDrag(Operation);
@@ -2854,10 +2858,20 @@ FReply SSubobjectEditor::TryHandleAssetDragDropOperation(const FDragDropEvent& D
 				TUniquePtr<FScopedTransaction> AddTransaction = MakeUnique<FScopedTransaction>(LOCTEXT("AddComponent", "Add Component"));
 
 				FAddNewSubobjectParams NewComponentParams;
+				
+				// Attach to the target item if any
+				if (InTargetItem)
+				{
+					NewComponentParams.ParentHandle = InTargetItem->GetDataHandle();
+				}
 
 				// Attach to the currently selected handle, if there are none then select the root object
-				const TArray<FSubobjectDataHandle>& SelectedHandles = GetSelectedHandles();
-				NewComponentParams.ParentHandle = SelectedHandles.IsEmpty() ? GetObjectContextHandle() : SelectedHandles[0];
+				if (!NewComponentParams.ParentHandle.IsValid())
+				{
+					const TArray<FSubobjectDataHandle>& SelectedHandles = GetSelectedHandles();
+					NewComponentParams.ParentHandle = SelectedHandles.IsEmpty() ? GetObjectContextHandle() : SelectedHandles[0];
+				}
+
 				NewComponentParams.BlueprintContext = ShouldModifyBPOnAssetDrop() ? GetBlueprint() : nullptr;
 
 				FText FailReason;
diff --git a/Engine/Source/Editor/SubobjectEditor/Public/SSubobjectEditor.h b/Engine/Source/Editor/SubobjectEditor/Public/SSubobjectEditor.h
--- a/Engine/Source/Editor/SubobjectEditor/Public/SSubobjectEditor.h
+++ b/Engine/Source/Editor/SubobjectEditor/Public/SSubobjectEditor.h
@@ -536,6 +536,9 @@ public:
 	/** Try to handle a drag-drop operation */
 	UE_API FReply TryHandleAssetDragDropOperation(const FDragDropEvent& DragDropEvent);
 	
+	/** Handle drag-drop operation on a tree item */
+	UE_API FReply TryHandleAssetDragDropOperation(const FDragDropEvent& InDragDropEvent, FSubobjectEditorTreeNodePtrType InTargetItem);
+	
 	/** Sets UI customizations of this SCSEditor. */
 	UE_API void SetUICustomization(TSharedPtr<ISCSEditorUICustomization> InUICustomization);
 	
